Moving NNE into the engine

- Moved model data factory to private
- Moved NNEEditorTools to NNEEditor inside Engine/Source/Editor (renaming all classes and files properly)
- Moved NNE to Engine/Source/Runtime
- Deleted the (now empty) NNE plugin and removed all references to it
- Added loading the NNEEditor module explicitely from runtimes that need NNE asset support
- Adjusted CreateModel and CreateModelInstance to use the proper suffix and return shared pointers (and updated all references in engine and plugins)

If you need to adapt your code:
- Remove the NNE plugin dependency (while keeping the NNE module dependency inside your Build.cs)
- Change all your CreateModel and CreateModelInstance to CreateModel??? and CreateModelInstance??? (with ??? being CPU, GPU or RDG) and change the return value from TUniquePtr to TSharedPtr

Testing:
- Clean workspace in UGS with 'Select  All'
- Update to the latest CL in UGS
- Checkout this CL
- Regenerate the project files
- Rebuild UE5
- Launch the NNEIncubator, run the style transfer and the session frontend tests

#rnx

diff --git a/Engine/Plugins/AI/MLAdapter/MLAdapter.uplugin b/Engine/Plugins/AI/MLAdapter/MLAdapter.uplugin
--- a/Engine/Plugins/AI/MLAdapter/MLAdapter.uplugin
+++ b/Engine/Plugins/AI/MLAdapter/MLAdapter.uplugin
@@ -35,10 +35,6 @@
 		{
 			"Name": "EnhancedInput",
 			"Enabled": true
-		},
-		{
-			"Name": "NNE",
-			"Enabled": true
 		}
 	]
 }
\ No newline at end of file
diff --git a/Engine/Plugins/AI/MLAdapter/Source/MLAdapter/Private/Agents/MLAdapterAgent_Inference.cpp b/Engine/Plugins/AI/MLAdapter/Source/MLAdapter/Private/Agents/MLAdapterAgent_Inference.cpp
--- a/Engine/Plugins/AI/MLAdapter/Source/MLAdapter/Private/Agents/MLAdapterAgent_Inference.cpp
+++ b/Engine/Plugins/AI/MLAdapter/Source/MLAdapter/Private/Agents/MLAdapterAgent_Inference.cpp
@@ -25,7 +25,7 @@ void UMLAdapterAgent_Inference::PostInitProperties()
 		return;
 	}
 
-	Brain = Runtime->CreateModel(ModelData)->CreateModelInstance();
+	Brain = Runtime->CreateModelCPU(ModelData)->CreateModelInstanceCPU();
 
 	TConstArrayView<UE::NNE::FTensorDesc> InputTensorDescs = Brain->GetInputTensorDescs();
 	if (InputTensorDescs.Num() != 1)
diff --git a/Engine/Plugins/AI/MLAdapter/Source/MLAdapter/Public/Agents/MLAdapterAgent_Inference.h b/Engine/Plugins/AI/MLAdapter/Source/MLAdapter/Public/Agents/MLAdapterAgent_Inference.h
--- a/Engine/Plugins/AI/MLAdapter/Source/MLAdapter/Public/Agents/MLAdapterAgent_Inference.h
+++ b/Engine/Plugins/AI/MLAdapter/Source/MLAdapter/Public/Agents/MLAdapterAgent_Inference.h
@@ -30,7 +30,7 @@ public:
 	TObjectPtr<UNNEModelData> ModelData;
 
 	/** The neural network controlling this agent. */
-	TUniquePtr<UE::NNE::IModelInstanceCPU> Brain;
+	TSharedPtr<UE::NNE::IModelInstanceCPU> Brain;
 
 private:
 	uint32 InputTensorSizeInBytes = 0;
diff --git a/Engine/Plugins/Animation/MLDeformer/VertexDeltaModel/Source/VertexDeltaModel/Private/VertexDeltaModelInstance.cpp b/Engine/Plugins/Animation/MLDeformer/VertexDeltaModel/Source/VertexDeltaModel/Private/VertexDeltaModelInstance.cpp
--- a/Engine/Plugins/Animation/MLDeformer/VertexDeltaModel/Source/VertexDeltaModel/Private/VertexDeltaModelInstance.cpp
+++ b/Engine/Plugins/Animation/MLDeformer/VertexDeltaModel/Source/VertexDeltaModel/Private/VertexDeltaModelInstance.cpp
@@ -177,7 +177,7 @@ void UVertexDeltaModelInstance::CreateNNEModel()
 				if (RuntimeRDG.IsValid())
 				{
 					// allocate tensor inputs and outputs
-					ModelInstanceRDG = RuntimeRDG->CreateModel(ModelData)->CreateModelInstance();
+					ModelInstanceRDG = RuntimeRDG->CreateModelRDG(ModelData)->CreateModelInstanceRDG();
 
 					if (ModelInstanceRDG)
 					{
diff --git a/Engine/Plugins/Animation/MLDeformer/VertexDeltaModel/Source/VertexDeltaModel/Public/VertexDeltaModelInstance.h b/Engine/Plugins/Animation/MLDeformer/VertexDeltaModel/Source/VertexDeltaModel/Public/VertexDeltaModelInstance.h
--- a/Engine/Plugins/Animation/MLDeformer/VertexDeltaModel/Source/VertexDeltaModel/Public/VertexDeltaModelInstance.h
+++ b/Engine/Plugins/Animation/MLDeformer/VertexDeltaModel/Source/VertexDeltaModel/Public/VertexDeltaModelInstance.h
@@ -60,7 +60,7 @@ private:
 	TRefCountPtr<FRDGPooledBuffer> RDGVertexDeltaBuffer;
 
 	// The NNE RDG Model 
-	TUniquePtr<UE::NNE::IModelInstanceRDG> ModelInstanceRDG;
+	TSharedPtr<UE::NNE::IModelInstanceRDG> ModelInstanceRDG;
 	
 	// The CPU Input Tensor Buffer
 	TArray<float> NNEInputTensorBuffer;
diff --git a/Engine/Plugins/Animation/MLDeformer/VertexDeltaModel/VertexDeltaModel.uplugin b/Engine/Plugins/Animation/MLDeformer/VertexDeltaModel/VertexDeltaModel.uplugin
--- a/Engine/Plugins/Animation/MLDeformer/VertexDeltaModel/VertexDeltaModel.uplugin
+++ b/Engine/Plugins/Animation/MLDeformer/VertexDeltaModel/VertexDeltaModel.uplugin
@@ -38,10 +38,6 @@
 			"Name": "GeometryCache",
 			"Enabled": true
 		},
-		{
-			"Name": "NNE",
-			"Enabled": true
-		},
 		{
 			"Name": "NNERuntimeRDG",
 			"Enabled": true
diff --git a/Engine/Plugins/Experimental/NNE/Config/BaseNNE.ini b/Engine/Plugins/Experimental/NNE/Config/BaseNNE.ini
deleted file mode 100644
diff --git a/Engine/Plugins/Experimental/NNE/NNE.uplugin b/Engine/Plugins/Experimental/NNE/NNE.uplugin
deleted file mode 100644
diff --git a/Engine/Plugins/Experimental/NNE/Source/NNE/Internal/NNEAttributeDataType.h b/Engine/Plugins/Experimental/NNE/Source/NNE/Internal/NNEAttributeDataType.h
deleted file mode 100644
diff --git a/Engine/Plugins/Experimental/NNE/Source/NNE/Internal/NNEAttributeMap.h b/Engine/Plugins/Experimental/NNE/Source/NNE/Internal/NNEAttributeMap.h
deleted file mode 100644
diff --git a/Engine/Plugins/Experimental/NNE/Source/NNE/Internal/NNEAttributeValue.h b/Engine/Plugins/Experimental/NNE/Source/NNE/Internal/NNEAttributeValue.h
deleted file mode 100644
diff --git a/Engine/Plugins/Experimental/NNE/Source/NNE/Internal/NNEAttributeValueTraits.h b/Engine/Plugins/Experimental/NNE/Source/NNE/Internal/NNEAttributeValueTraits.h
deleted file mode 100644
diff --git a/Engine/Plugins/Experimental/NNE/Source/NNE/Internal/NNEModelBase.h b/Engine/Plugins/Experimental/NNE/Source/NNE/Internal/NNEModelBase.h
deleted file mode 100644
diff --git a/Engine/Plugins/Experimental/NNE/Source/NNE/Internal/NNEModelOptimizerInterface.h b/Engine/Plugins/Experimental/NNE/Source/NNE/Internal/NNEModelOptimizerInterface.h
deleted file mode 100644
diff --git a/Engine/Plugins/Experimental/NNE/Source/NNE/Internal/NNEOperator.h b/Engine/Plugins/Experimental/NNE/Source/NNE/Internal/NNEOperator.h
deleted file mode 100644
diff --git a/Engine/Plugins/Experimental/NNE/Source/NNE/Internal/NNERuntimeFormat.h b/Engine/Plugins/Experimental/NNE/Source/NNE/Internal/NNERuntimeFormat.h
deleted file mode 100644
diff --git a/Engine/Plugins/Experimental/NNE/Source/NNE/Internal/NNETensor.h b/Engine/Plugins/Experimental/NNE/Source/NNE/Internal/NNETensor.h
deleted file mode 100644
diff --git a/Engine/Plugins/Experimental/NNE/Source/NNE/NNE.Build.cs b/Engine/Plugins/Experimental/NNE/Source/NNE/NNE.Build.cs
deleted file mode 100644
diff --git a/Engine/Plugins/Experimental/NNE/Source/NNE/Private/NNE.cpp b/Engine/Plugins/Experimental/NNE/Source/NNE/Private/NNE.cpp
deleted file mode 100644
diff --git a/Engine/Plugins/Experimental/NNE/Source/NNE/Private/NNEModelData.cpp b/Engine/Plugins/Experimental/NNE/Source/NNE/Private/NNEModelData.cpp
deleted file mode 100644
diff --git a/Engine/Plugins/Experimental/NNE/Source/NNE/Private/NNEModule.cpp b/Engine/Plugins/Experimental/NNE/Source/NNE/Private/NNEModule.cpp
deleted file mode 100644
diff --git a/Engine/Plugins/Experimental/NNE/Source/NNE/Private/NNETypes.cpp b/Engine/Plugins/Experimental/NNE/Source/NNE/Private/NNETypes.cpp
deleted file mode 100644
diff --git a/Engine/Plugins/Experimental/NNE/Source/NNE/Public/NNE.h b/Engine/Plugins/Experimental/NNE/Source/NNE/Public/NNE.h
deleted file mode 100644
diff --git a/Engine/Plugins/Experimental/NNE/Source/NNE/Public/NNEModelData.h b/Engine/Plugins/Experimental/NNE/Source/NNE/Public/NNEModelData.h
deleted file mode 100644
diff --git a/Engine/Plugins/Experimental/NNE/Source/NNE/Public/NNERuntime.h b/Engine/Plugins/Experimental/NNE/Source/NNE/Public/NNERuntime.h
deleted file mode 100644
diff --git a/Engine/Plugins/Experimental/NNE/Source/NNE/Public/NNERuntimeCPU.h b/Engine/Plugins/Experimental/NNE/Source/NNE/Public/NNERuntimeCPU.h
deleted file mode 100644
diff --git a/Engine/Plugins/Experimental/NNE/Source/NNE/Public/NNERuntimeGPU.h b/Engine/Plugins/Experimental/NNE/Source/NNE/Public/NNERuntimeGPU.h
deleted file mode 100644
diff --git a/Engine/Plugins/Experimental/NNE/Source/NNE/Public/NNERuntimeRDG.h b/Engine/Plugins/Experimental/NNE/Source/NNE/Public/NNERuntimeRDG.h
deleted file mode 100644
diff --git a/Engine/Plugins/Experimental/NNE/Source/NNE/Public/NNETypes.h b/Engine/Plugins/Experimental/NNE/Source/NNE/Public/NNETypes.h
deleted file mode 100644
diff --git a/Engine/Plugins/Experimental/NNE/Source/NNEEditorTools/NNEEditorTools.Build.cs b/Engine/Plugins/Experimental/NNE/Source/NNEEditorTools/NNEEditorTools.Build.cs
deleted file mode 100644
diff --git a/Engine/Plugins/Experimental/NNE/Source/NNEEditorTools/Private/NNEEditorTools.cpp b/Engine/Plugins/Experimental/NNE/Source/NNEEditorTools/Private/NNEEditorTools.cpp
deleted file mode 100644
diff --git a/Engine/Plugins/Experimental/NNE/Source/NNEEditorTools/Private/NNEEditorToolsModelDataActions.cpp b/Engine/Plugins/Experimental/NNE/Source/NNEEditorTools/Private/NNEEditorToolsModelDataActions.cpp
deleted file mode 100644
diff --git a/Engine/Plugins/Experimental/NNE/Source/NNEEditorTools/Private/NNEEditorToolsModelDataActions.h b/Engine/Plugins/Experimental/NNE/Source/NNEEditorTools/Private/NNEEditorToolsModelDataActions.h
deleted file mode 100644
diff --git a/Engine/Plugins/Experimental/NNE/Source/NNEEditorTools/Private/NNEEditorToolsModelDataEditorToolkit.cpp b/Engine/Plugins/Experimental/NNE/Source/NNEEditorTools/Private/NNEEditorToolsModelDataEditorToolkit.cpp
deleted file mode 100644
diff --git a/Engine/Plugins/Experimental/NNE/Source/NNEEditorTools/Private/NNEEditorToolsModelDataEditorToolkit.h b/Engine/Plugins/Experimental/NNE/Source/NNEEditorTools/Private/NNEEditorToolsModelDataEditorToolkit.h
deleted file mode 100644
diff --git a/Engine/Plugins/Experimental/NNE/Source/NNEEditorTools/Private/NNEModelDataFactory.cpp b/Engine/Plugins/Experimental/NNE/Source/NNEEditorTools/Private/NNEModelDataFactory.cpp
deleted file mode 100644
diff --git a/Engine/Plugins/Experimental/NNE/Source/NNEEditorTools/Private/SNNEEditorToolsRuntimesWidget.cpp b/Engine/Plugins/Experimental/NNE/Source/NNEEditorTools/Private/SNNEEditorToolsRuntimesWidget.cpp
deleted file mode 100644
diff --git a/Engine/Plugins/Experimental/NNE/Source/NNEEditorTools/Private/SNNEEditorToolsRuntimesWidget.h b/Engine/Plugins/Experimental/NNE/Source/NNEEditorTools/Private/SNNEEditorToolsRuntimesWidget.h
deleted file mode 100644
diff --git a/Engine/Plugins/Experimental/NNE/Source/NNEEditorTools/Public/NNEModelDataFactory.h b/Engine/Plugins/Experimental/NNE/Source/NNEEditorTools/Public/NNEModelDataFactory.h
deleted file mode 100644
diff --git a/Engine/Plugins/Experimental/NNERuntimeORTCpu/NNERuntimeORTCpu.uplugin b/Engine/Plugins/Experimental/NNERuntimeORTCpu/NNERuntimeORTCpu.uplugin
--- a/Engine/Plugins/Experimental/NNERuntimeORTCpu/NNERuntimeORTCpu.uplugin
+++ b/Engine/Plugins/Experimental/NNERuntimeORTCpu/NNERuntimeORTCpu.uplugin
@@ -49,11 +49,5 @@
 			"LoadingPhase": "Default",
 			"PlatformAllowList": [ "Win64", "Linux", "Mac" ]
 		}
-	],
-	"Plugins": [
-		{
-			"Name": "NNE",
-			"Enabled": true
-		}
 	]
 }
\ No newline at end of file
diff --git a/Engine/Plugins/Experimental/NNERuntimeORTCpu/Source/NNERuntimeORTCpu/NNERuntimeORTCpu.Build.cs b/Engine/Plugins/Experimental/NNERuntimeORTCpu/Source/NNERuntimeORTCpu/NNERuntimeORTCpu.Build.cs
--- a/Engine/Plugins/Experimental/NNERuntimeORTCpu/Source/NNERuntimeORTCpu/NNERuntimeORTCpu.Build.cs
+++ b/Engine/Plugins/Experimental/NNERuntimeORTCpu/Source/NNERuntimeORTCpu/NNERuntimeORTCpu.Build.cs
@@ -33,5 +33,10 @@ public class NNERuntimeORTCpu : ModuleRules
 				"NNEOnnxruntime"
 			}
 		);
+
+		if (Target.Type == TargetType.Editor)
+		{
+			PublicDependencyModuleNames.Add("NNEEditor");
+		}
 	}
 }
diff --git a/Engine/Plugins/Experimental/NNERuntimeORTCpu/Source/NNERuntimeORTCpu/Private/NNERuntimeORTCpu.cpp b/Engine/Plugins/Experimental/NNERuntimeORTCpu/Source/NNERuntimeORTCpu/Private/NNERuntimeORTCpu.cpp
--- a/Engine/Plugins/Experimental/NNERuntimeORTCpu/Source/NNERuntimeORTCpu/Private/NNERuntimeORTCpu.cpp
+++ b/Engine/Plugins/Experimental/NNERuntimeORTCpu/Source/NNERuntimeORTCpu/Private/NNERuntimeORTCpu.cpp
@@ -73,13 +73,13 @@ bool UNNERuntimeORTCpuImpl::CanCreateModelCPU(TObjectPtr<UNNEModelData> ModelDat
 	return bResult;
 }
 
-TUniquePtr<UE::NNE::IModelCPU> UNNERuntimeORTCpuImpl::CreateModel(TObjectPtr<UNNEModelData> ModelData)
+TSharedPtr<UE::NNE::IModelCPU> UNNERuntimeORTCpuImpl::CreateModelCPU(TObjectPtr<UNNEModelData> ModelData)
 {
 	check(ModelData != nullptr);
 	
 	if (!CanCreateModelCPU(ModelData))
 	{
-		return TUniquePtr<UE::NNE::IModelCPU>();
+		return TSharedPtr<UE::NNE::IModelCPU>();
 	}
 
 	TConstArrayView<uint8> Data = ModelData->GetModelData(GetRuntimeName());
@@ -96,5 +96,5 @@ TUniquePtr<UE::NNE::IModelCPU> UNNERuntimeORTCpuImpl::CreateModel(TObjectPtr<UNN
 		FEngineAnalytics::GetProvider().RecordEvent(TEXT("NeuralNetworkEngine.CreateModel"), Attributes);
 	}
 
-	return TUniquePtr<UE::NNE::IModelCPU>(IModel);
+	return TSharedPtr<UE::NNE::IModelCPU>(IModel);
 }
\ No newline at end of file
diff --git a/Engine/Plugins/Experimental/NNERuntimeORTCpu/Source/NNERuntimeORTCpu/Private/NNERuntimeORTCpu.h b/Engine/Plugins/Experimental/NNERuntimeORTCpu/Source/NNERuntimeORTCpu/Private/NNERuntimeORTCpu.h
--- a/Engine/Plugins/Experimental/NNERuntimeORTCpu/Source/NNERuntimeORTCpu/Private/NNERuntimeORTCpu.h
+++ b/Engine/Plugins/Experimental/NNERuntimeORTCpu/Source/NNERuntimeORTCpu/Private/NNERuntimeORTCpu.h
@@ -39,5 +39,5 @@ public:
 	virtual FString GetModelDataIdentifier(FString FileType, TConstArrayView<uint8> FileData, FGuid FileId, const ITargetPlatform* TargetPlatform) override;
 
 	virtual bool CanCreateModelCPU(TObjectPtr<UNNEModelData> ModelData) const override;
-	virtual TUniquePtr<UE::NNE::IModelCPU> CreateModel(TObjectPtr<UNNEModelData> ModelData) override;
+	virtual TSharedPtr<UE::NNE::IModelCPU> CreateModelCPU(TObjectPtr<UNNEModelData> ModelData) override;
 };
\ No newline at end of file
diff --git a/Engine/Plugins/Experimental/NNERuntimeORTCpu/Source/NNERuntimeORTCpu/Private/NNERuntimeORTCpuModel.cpp b/Engine/Plugins/Experimental/NNERuntimeORTCpu/Source/NNERuntimeORTCpu/Private/NNERuntimeORTCpuModel.cpp
--- a/Engine/Plugins/Experimental/NNERuntimeORTCpu/Source/NNERuntimeORTCpu/Private/NNERuntimeORTCpuModel.cpp
+++ b/Engine/Plugins/Experimental/NNERuntimeORTCpu/Source/NNERuntimeORTCpu/Private/NNERuntimeORTCpuModel.cpp
@@ -300,7 +300,7 @@ namespace UE::NNERuntimeORTCpu::Private
 		InputTransferStatisticsEstimator.ResetStats();
 	}
 
-	TUniquePtr<UE::NNE::IModelInstanceCPU> FModelCPU::CreateModelInstance()
+	TSharedPtr<UE::NNE::IModelInstanceCPU> FModelCPU::CreateModelInstanceCPU()
 	{
 		const UE::NNERuntimeORTCpu::Private::FRuntimeConf InConf;
 		UE::NNERuntimeORTCpu::Private::FModelInstanceCPU* ModelInstance = new UE::NNERuntimeORTCpu::Private::FModelInstanceCPU(ORTEnvironment, InConf);
@@ -308,11 +308,11 @@ namespace UE::NNERuntimeORTCpu::Private
 		if (!ModelInstance->Init(ModelData))
 		{
 			delete ModelInstance;
-			return TUniquePtr<UE::NNE::IModelInstanceCPU>();
+			return TSharedPtr<UE::NNE::IModelInstanceCPU>();
 		}
 
 		UE::NNE::IModelInstanceCPU* IModelInstance = static_cast<UE::NNE::IModelInstanceCPU*>(ModelInstance);
-		return TUniquePtr<UE::NNE::IModelInstanceCPU>(IModelInstance);
+		return TSharedPtr<UE::NNE::IModelInstanceCPU>(IModelInstance);
 	}
 
 	FModelCPU::FModelCPU(Ort::Env* InORTEnvironment, TConstArrayView<uint8> InModelData) :
diff --git a/Engine/Plugins/Experimental/NNERuntimeORTCpu/Source/NNERuntimeORTCpu/Private/NNERuntimeORTCpuModel.h b/Engine/Plugins/Experimental/NNERuntimeORTCpu/Source/NNERuntimeORTCpu/Private/NNERuntimeORTCpuModel.h
--- a/Engine/Plugins/Experimental/NNERuntimeORTCpu/Source/NNERuntimeORTCpu/Private/NNERuntimeORTCpuModel.h
+++ b/Engine/Plugins/Experimental/NNERuntimeORTCpu/Source/NNERuntimeORTCpu/Private/NNERuntimeORTCpuModel.h
@@ -86,7 +86,7 @@ namespace UE::NNERuntimeORTCpu::Private
 		FModelCPU(Ort::Env* InORTEnvironment, TConstArrayView<uint8> ModelData);
 		virtual ~FModelCPU() {};
 
-		virtual TUniquePtr<UE::NNE::IModelInstanceCPU> CreateModelInstance() override;
+		virtual TSharedPtr<UE::NNE::IModelInstanceCPU> CreateModelInstanceCPU() override;
 
 	private:
 		Ort::Env* ORTEnvironment;
diff --git a/Engine/Plugins/Experimental/NNERuntimeORTCpu/Source/NNERuntimeORTCpu/Private/NNERuntimeORTCpuModule.cpp b/Engine/Plugins/Experimental/NNERuntimeORTCpu/Source/NNERuntimeORTCpu/Private/NNERuntimeORTCpuModule.cpp
--- a/Engine/Plugins/Experimental/NNERuntimeORTCpu/Source/NNERuntimeORTCpu/Private/NNERuntimeORTCpuModule.cpp
+++ b/Engine/Plugins/Experimental/NNERuntimeORTCpu/Source/NNERuntimeORTCpu/Private/NNERuntimeORTCpuModule.cpp
@@ -7,6 +7,10 @@
 
 void FNNERuntimeORTCpuModule::StartupModule()
 {
+#ifdef WITH_EDITOR
+	FModuleManager::Get().LoadModule(TEXT("NNEEditor"));
+#endif
+
 	// NNE runtime startup
 	NNERuntimeORTCpu = NewObject<UNNERuntimeORTCpuImpl>();
 	if (NNERuntimeORTCpu.IsValid())
diff --git a/Engine/Plugins/Experimental/NNERuntimeORTGpu/NNERuntimeORTGpu.uplugin b/Engine/Plugins/Experimental/NNERuntimeORTGpu/NNERuntimeORTGpu.uplugin
--- a/Engine/Plugins/Experimental/NNERuntimeORTGpu/NNERuntimeORTGpu.uplugin
+++ b/Engine/Plugins/Experimental/NNERuntimeORTGpu/NNERuntimeORTGpu.uplugin
@@ -27,10 +27,6 @@
 		}
 	],
 	"Plugins": [
-		{
-			"Name": "NNE",
-			"Enabled": true
-		},
 		{
 			"Name": "NNERuntimeORTCpu",
 			"Enabled": true
diff --git a/Engine/Plugins/Experimental/NNERuntimeORTGpu/Source/NNERuntimeORT/NNERuntimeORT.Build.cs b/Engine/Plugins/Experimental/NNERuntimeORTGpu/Source/NNERuntimeORT/NNERuntimeORT.Build.cs
--- a/Engine/Plugins/Experimental/NNERuntimeORTGpu/Source/NNERuntimeORT/NNERuntimeORT.Build.cs
+++ b/Engine/Plugins/Experimental/NNERuntimeORTGpu/Source/NNERuntimeORT/NNERuntimeORT.Build.cs
@@ -57,5 +57,10 @@ public class NNERuntimeORT : ModuleRules
 			AddEngineThirdPartyPrivateStaticDependencies(Target, "DirectML");
 
 		}
+
+		if (Target.Type == TargetType.Editor)
+		{
+			PublicDependencyModuleNames.Add("NNEEditor");
+		}
 	}
 }
diff --git a/Engine/Plugins/Experimental/NNERuntimeORTGpu/Source/NNERuntimeORT/Private/NNERuntimeORT.cpp b/Engine/Plugins/Experimental/NNERuntimeORTGpu/Source/NNERuntimeORT/Private/NNERuntimeORT.cpp
--- a/Engine/Plugins/Experimental/NNERuntimeORTGpu/Source/NNERuntimeORT/Private/NNERuntimeORT.cpp
+++ b/Engine/Plugins/Experimental/NNERuntimeORTGpu/Source/NNERuntimeORT/Private/NNERuntimeORT.cpp
@@ -98,14 +98,14 @@ bool UNNERuntimeORTGpuImpl::CanCreateModelGPU(TObjectPtr<UNNEModelData> ModelDat
 	return bResult;
 }
 
-TUniquePtr<UE::NNE::IModelGPU> UNNERuntimeORTGpuImpl::CreateModel(TObjectPtr<UNNEModelData> ModelData)
+TSharedPtr<UE::NNE::IModelGPU> UNNERuntimeORTGpuImpl::CreateModelGPU(TObjectPtr<UNNEModelData> ModelData)
 {
 	check(ModelData != nullptr);
 	check(ORTEnvironment.IsValid());
 
 	if (!CanCreateModelGPU(ModelData))
 	{
-		return TUniquePtr<UE::NNE::IModelGPU>();
+		return TSharedPtr<UE::NNE::IModelGPU>();
 	}
 
 	UE::NNE::IModelGPU* IModel = nullptr;
@@ -121,7 +121,7 @@ TUniquePtr<UE::NNE::IModelGPU> UNNERuntimeORTGpuImpl::CreateModel(TObjectPtr<UNN
 			break;
 		default:
 			UE_LOG(LogNNE, Error, TEXT("Failed to create model for ORT GPU runtime, unsupported provider. Runtime will not be functional."));
-			return TUniquePtr<UE::NNE::IModelGPU>();
+			return TSharedPtr<UE::NNE::IModelGPU>();
 	}
 
 	if (FEngineAnalytics::IsAvailable())
@@ -134,7 +134,7 @@ TUniquePtr<UE::NNE::IModelGPU> UNNERuntimeORTGpuImpl::CreateModel(TObjectPtr<UNN
 		FEngineAnalytics::GetProvider().RecordEvent(TEXT("NeuralNetworkEngine.CreateModel"), Attributes);
 	}
 
-	return TUniquePtr<UE::NNE::IModelGPU>(IModel);
+	return TSharedPtr<UE::NNE::IModelGPU>(IModel);
 }
 
 #else // PLATFORM_WINDOWS
diff --git a/Engine/Plugins/Experimental/NNERuntimeORTGpu/Source/NNERuntimeORT/Private/NNERuntimeORT.h b/Engine/Plugins/Experimental/NNERuntimeORTGpu/Source/NNERuntimeORT/Private/NNERuntimeORT.h
--- a/Engine/Plugins/Experimental/NNERuntimeORTGpu/Source/NNERuntimeORT/Private/NNERuntimeORT.h
+++ b/Engine/Plugins/Experimental/NNERuntimeORTGpu/Source/NNERuntimeORT/Private/NNERuntimeORT.h
@@ -51,5 +51,5 @@ public:
 	virtual FString GetModelDataIdentifier(FString FileType, TConstArrayView<uint8> FileData, FGuid FileId, const ITargetPlatform* TargetPlatform) override;
 
 	virtual bool CanCreateModelGPU(TObjectPtr<UNNEModelData> ModelData) const override;
-	virtual TUniquePtr<UE::NNE::IModelGPU> CreateModel(TObjectPtr<UNNEModelData> ModelData) override;
+	virtual TSharedPtr<UE::NNE::IModelGPU> CreateModelGPU(TObjectPtr<UNNEModelData> ModelData) override;
 };
\ No newline at end of file
diff --git a/Engine/Plugins/Experimental/NNERuntimeORTGpu/Source/NNERuntimeORT/Private/NNERuntimeORTModel.cpp b/Engine/Plugins/Experimental/NNERuntimeORTGpu/Source/NNERuntimeORT/Private/NNERuntimeORTModel.cpp
--- a/Engine/Plugins/Experimental/NNERuntimeORTGpu/Source/NNERuntimeORT/Private/NNERuntimeORTModel.cpp
+++ b/Engine/Plugins/Experimental/NNERuntimeORTGpu/Source/NNERuntimeORT/Private/NNERuntimeORTModel.cpp
@@ -327,7 +327,7 @@ namespace UE::NNERuntimeORT::Private
 
 #if PLATFORM_WINDOWS
 
-	TUniquePtr<UE::NNE::IModelInstanceGPU> FModelORTDml::CreateModelInstance()
+	TSharedPtr<UE::NNE::IModelInstanceGPU> FModelORTDml::CreateModelInstanceGPU()
 	{
 		const FRuntimeConf InConf;
 		FModelInstanceORTDml* ModelInstance = new FModelInstanceORTDml(ORTEnvironment, InConf);
@@ -335,11 +335,11 @@ namespace UE::NNERuntimeORT::Private
 		if (!ModelInstance->Init(ModelData))
 		{
 			delete ModelInstance;
-			return TUniquePtr<UE::NNE::IModelInstanceGPU>();
+			return TSharedPtr<UE::NNE::IModelInstanceGPU>();
 		}
 
 		UE::NNE::IModelInstanceGPU* IModelInstance = static_cast<UE::NNE::IModelInstanceGPU*>(ModelInstance);
-		return TUniquePtr<UE::NNE::IModelInstanceGPU>(IModelInstance);
+		return TSharedPtr<UE::NNE::IModelInstanceGPU>(IModelInstance);
 	}
 
 	FModelORTDml::FModelORTDml(Ort::Env* InORTEnvironment, TConstArrayView<uint8> InModelData) :
@@ -349,7 +349,7 @@ namespace UE::NNERuntimeORT::Private
 
 	}
 
-	TUniquePtr<UE::NNE::IModelInstanceGPU> FModelORTCuda::CreateModelInstance()
+	TSharedPtr<UE::NNE::IModelInstanceGPU> FModelORTCuda::CreateModelInstanceGPU()
 	{
 		const FRuntimeConf InConf;
 		FModelInstanceORTCuda* ModelInstance = new FModelInstanceORTCuda(ORTEnvironment, InConf);
@@ -357,11 +357,11 @@ namespace UE::NNERuntimeORT::Private
 		if (!ModelInstance->Init(ModelData))
 		{
 			delete ModelInstance;
-			return TUniquePtr<UE::NNE::IModelInstanceGPU>();
+			return TSharedPtr<UE::NNE::IModelInstanceGPU>();
 		}
 
 		UE::NNE::IModelInstanceGPU* IModelInstance = static_cast<UE::NNE::IModelInstanceGPU*>(ModelInstance);
-		return TUniquePtr<UE::NNE::IModelInstanceGPU>(IModelInstance);
+		return TSharedPtr<UE::NNE::IModelInstanceGPU>(IModelInstance);
 	}
 
 	FModelORTCuda::FModelORTCuda(Ort::Env* InORTEnvironment, TConstArrayView<uint8> InModelData) :
diff --git a/Engine/Plugins/Experimental/NNERuntimeORTGpu/Source/NNERuntimeORT/Private/NNERuntimeORTModel.h b/Engine/Plugins/Experimental/NNERuntimeORTGpu/Source/NNERuntimeORT/Private/NNERuntimeORTModel.h
--- a/Engine/Plugins/Experimental/NNERuntimeORTGpu/Source/NNERuntimeORT/Private/NNERuntimeORTModel.h
+++ b/Engine/Plugins/Experimental/NNERuntimeORTGpu/Source/NNERuntimeORT/Private/NNERuntimeORTModel.h
@@ -107,7 +107,7 @@ namespace UE::NNERuntimeORT::Private
 		FModelORTDml(Ort::Env* InORTEnvironment, TConstArrayView<uint8> ModelData);
 		virtual ~FModelORTDml() {};
 
-		virtual TUniquePtr<UE::NNE::IModelInstanceGPU> CreateModelInstance() override;
+		virtual TSharedPtr<UE::NNE::IModelInstanceGPU> CreateModelInstanceGPU() override;
 
 	private:
 		Ort::Env* ORTEnvironment;
@@ -120,7 +120,7 @@ namespace UE::NNERuntimeORT::Private
 		FModelORTCuda(Ort::Env* InORTEnvironment, TConstArrayView<uint8> ModelData);
 		virtual ~FModelORTCuda() {};
 
-		virtual TUniquePtr<UE::NNE::IModelInstanceGPU> CreateModelInstance() override;
+		virtual TSharedPtr<UE::NNE::IModelInstanceGPU> CreateModelInstanceGPU() override;
 
 	private:
 		Ort::Env* ORTEnvironment;
diff --git a/Engine/Plugins/Experimental/NNERuntimeORTGpu/Source/NNERuntimeORT/Private/NNERuntimeORTModule.cpp b/Engine/Plugins/Experimental/NNERuntimeORTGpu/Source/NNERuntimeORT/Private/NNERuntimeORTModule.cpp
--- a/Engine/Plugins/Experimental/NNERuntimeORTGpu/Source/NNERuntimeORT/Private/NNERuntimeORTModule.cpp
+++ b/Engine/Plugins/Experimental/NNERuntimeORTGpu/Source/NNERuntimeORT/Private/NNERuntimeORTModule.cpp
@@ -43,6 +43,10 @@ namespace UE::NNERuntimeORT::Private::DllHelper
 
 void FNNERuntimeORTModule::StartupModule()
 {
+#ifdef WITH_EDITOR
+	FModuleManager::Get().LoadModule(TEXT("NNEEditor"));
+#endif
+
 #if PLATFORM_WINDOWS
 	const FString PluginDir = IPluginManager::Get().FindPlugin("NNERuntimeORTGpu")->GetBaseDir();
 	const FString OrtBinPath = FPaths::Combine(PluginDir, TEXT(PREPROCESSOR_TO_STRING(ONNXRUNTIME_PLATFORM_PATH)));
diff --git a/Engine/Plugins/Experimental/NNERuntimeRDG/NNERuntimeRDG.uplugin b/Engine/Plugins/Experimental/NNERuntimeRDG/NNERuntimeRDG.uplugin
--- a/Engine/Plugins/Experimental/NNERuntimeRDG/NNERuntimeRDG.uplugin
+++ b/Engine/Plugins/Experimental/NNERuntimeRDG/NNERuntimeRDG.uplugin
@@ -33,10 +33,6 @@
 		}
 	],
 	"Plugins": [
-		{
-			"Name": "NNE",
-			"Enabled": true
-		},
 		{
 			"Name": "NNERuntimeORTCpu",
 			"Enabled": true
diff --git a/Engine/Plugins/Experimental/NNERuntimeRDG/Source/NNERuntimeRDG/NNERuntimeRDG.Build.cs b/Engine/Plugins/Experimental/NNERuntimeRDG/Source/NNERuntimeRDG/NNERuntimeRDG.Build.cs
--- a/Engine/Plugins/Experimental/NNERuntimeRDG/Source/NNERuntimeRDG/NNERuntimeRDG.Build.cs
+++ b/Engine/Plugins/Experimental/NNERuntimeRDG/Source/NNERuntimeRDG/NNERuntimeRDG.Build.cs
@@ -52,5 +52,10 @@ public class NNERuntimeRDG : ModuleRules
 		{	
 			PrivateDependencyModuleNames.Add("VulkanRHI");
 		}
+
+		if (Target.Type == TargetType.Editor)
+		{
+			PublicDependencyModuleNames.Add("NNEEditor");
+		}
 	}
 }
diff --git a/Engine/Plugins/Experimental/NNERuntimeRDG/Source/NNERuntimeRDG/Private/Dml/NNEDmlModel.cpp b/Engine/Plugins/Experimental/NNERuntimeRDG/Source/NNERuntimeRDG/Private/Dml/NNEDmlModel.cpp
--- a/Engine/Plugins/Experimental/NNERuntimeRDG/Source/NNERuntimeRDG/Private/Dml/NNEDmlModel.cpp
+++ b/Engine/Plugins/Experimental/NNERuntimeRDG/Source/NNERuntimeRDG/Private/Dml/NNEDmlModel.cpp
@@ -1386,18 +1386,18 @@ int FModelInstance::PrepareTensorShapesAndData()
 	return 0;
 }
 
-TUniquePtr<UE::NNE::IModelInstanceRDG> FModel::CreateModelInstance()
+TSharedPtr<UE::NNE::IModelInstanceRDG> FModel::CreateModelInstanceRDG()
 {
 	FModelInstance* ModelInstance = new FModelInstance();
 
 	if (!ModelInstance->Init(ModelData, DevCtx))
 	{
 		delete ModelInstance;
-		return TUniquePtr<UE::NNE::IModelInstanceRDG>();
+		return TSharedPtr<UE::NNE::IModelInstanceRDG>();
 	}
 
 	UE::NNE::IModelInstanceRDG* IModelInstance = static_cast<UE::NNE::IModelInstanceRDG*>(ModelInstance);
-	return TUniquePtr<UE::NNE::IModelInstanceRDG>(IModelInstance);
+	return TSharedPtr<UE::NNE::IModelInstanceRDG>(IModelInstance);
 }
 
 FModel::FModel(TConstArrayView<uint8> InModelData, FDmlDeviceContext* InDevCtx) 
diff --git a/Engine/Plugins/Experimental/NNERuntimeRDG/Source/NNERuntimeRDG/Private/Dml/NNEDmlModel.h b/Engine/Plugins/Experimental/NNERuntimeRDG/Source/NNERuntimeRDG/Private/Dml/NNEDmlModel.h
--- a/Engine/Plugins/Experimental/NNERuntimeRDG/Source/NNERuntimeRDG/Private/Dml/NNEDmlModel.h
+++ b/Engine/Plugins/Experimental/NNERuntimeRDG/Source/NNERuntimeRDG/Private/Dml/NNEDmlModel.h
@@ -98,7 +98,7 @@ public:
 	FModel(TConstArrayView<uint8> InModelData, FDmlDeviceContext* InDevCtx);
 	virtual ~FModel() {};
 
-	virtual TUniquePtr<UE::NNE::IModelInstanceRDG> CreateModelInstance() override;
+	virtual TSharedPtr<UE::NNE::IModelInstanceRDG> CreateModelInstanceRDG() override;
 
 private:
 	TArray<uint8> ModelData;
diff --git a/Engine/Plugins/Experimental/NNERuntimeRDG/Source/NNERuntimeRDG/Private/NNERuntimeRDGDml.cpp b/Engine/Plugins/Experimental/NNERuntimeRDG/Source/NNERuntimeRDG/Private/NNERuntimeRDGDml.cpp
--- a/Engine/Plugins/Experimental/NNERuntimeRDG/Source/NNERuntimeRDG/Private/NNERuntimeRDGDml.cpp
+++ b/Engine/Plugins/Experimental/NNERuntimeRDG/Source/NNERuntimeRDG/Private/NNERuntimeRDGDml.cpp
@@ -157,12 +157,12 @@ bool UNNERuntimeRDGDmlImpl::CanCreateModelRDG(TObjectPtr<UNNEModelData> ModelDat
 #endif
 };
 
-TUniquePtr<UE::NNE::IModelRDG> UNNERuntimeRDGDmlImpl::CreateModel(TObjectPtr<UNNEModelData> ModelData)
+TSharedPtr<UE::NNE::IModelRDG> UNNERuntimeRDGDmlImpl::CreateModelRDG(TObjectPtr<UNNEModelData> ModelData)
 {
 #ifdef NNE_USE_DIRECTML
 	if (!CanCreateModelRDG(ModelData))
 	{
-		return TUniquePtr<UE::NNE::IModelRDG>();
+		return TSharedPtr<UE::NNE::IModelRDG>();
 	}
 
 	TConstArrayView<uint8> Data = ModelData->GetModelData(GetRuntimeName());
@@ -178,9 +178,9 @@ TUniquePtr<UE::NNE::IModelRDG> UNNERuntimeRDGDmlImpl::CreateModel(TObjectPtr<UNN
 		FEngineAnalytics::GetProvider().RecordEvent(TEXT("NeuralNetworkEngine.CreateModel"), Attributes);
 	}
 
-	return TUniquePtr<UE::NNE::IModelRDG>(Model);
+	return TSharedPtr<UE::NNE::IModelRDG>(Model);
 #else
-	return TUniquePtr<UE::NNE::IModelRDG>();
+	return TSharedPtr<UE::NNE::IModelRDG>();
 #endif
 }
 
diff --git a/Engine/Plugins/Experimental/NNERuntimeRDG/Source/NNERuntimeRDG/Private/NNERuntimeRDGDml.h b/Engine/Plugins/Experimental/NNERuntimeRDG/Source/NNERuntimeRDG/Private/NNERuntimeRDGDml.h
--- a/Engine/Plugins/Experimental/NNERuntimeRDG/Source/NNERuntimeRDG/Private/NNERuntimeRDGDml.h
+++ b/Engine/Plugins/Experimental/NNERuntimeRDG/Source/NNERuntimeRDG/Private/NNERuntimeRDGDml.h
@@ -39,7 +39,7 @@ public:
 	virtual FString GetModelDataIdentifier(FString FileType, TConstArrayView<uint8> FileData, FGuid FileId, const ITargetPlatform* TargetPlatform) override;
 
 	virtual bool CanCreateModelRDG(TObjectPtr<UNNEModelData> ModelData) const override;
-	virtual TUniquePtr<UE::NNE::IModelRDG> CreateModel(TObjectPtr<UNNEModelData> ModelData) override;
+	virtual TSharedPtr<UE::NNE::IModelRDG> CreateModelRDG(TObjectPtr<UNNEModelData> ModelData) override;
 private:
 	UE::NNERuntimeRDG::Private::Dml::FDmlDeviceContext* Ctx;
 };
diff --git a/Engine/Plugins/Experimental/NNERuntimeRDG/Source/NNERuntimeRDG/Private/NNERuntimeRDGHlsl.cpp b/Engine/Plugins/Experimental/NNERuntimeRDG/Source/NNERuntimeRDG/Private/NNERuntimeRDGHlsl.cpp
--- a/Engine/Plugins/Experimental/NNERuntimeRDG/Source/NNERuntimeRDG/Private/NNERuntimeRDGHlsl.cpp
+++ b/Engine/Plugins/Experimental/NNERuntimeRDG/Source/NNERuntimeRDG/Private/NNERuntimeRDGHlsl.cpp
@@ -127,11 +127,11 @@ FString UNNERuntimeRDGHlslImpl::GetModelDataIdentifier(FString FileType, TConstA
 	return FileId.ToString(EGuidFormats::Digits) + "-" + UNNERuntimeRDGHlslImpl::GUID.ToString(EGuidFormats::Digits) + "-" + FString::FromInt(UNNERuntimeRDGHlslImpl::Version);
 }
 
-TUniquePtr<UE::NNE::IModelRDG> UNNERuntimeRDGHlslImpl::CreateModel(TObjectPtr<UNNEModelData> ModelData)
+TSharedPtr<UE::NNE::IModelRDG> UNNERuntimeRDGHlslImpl::CreateModelRDG(TObjectPtr<UNNEModelData> ModelData)
 {
 	if (!CanCreateModelRDG(ModelData))
 	{
-		return TUniquePtr<UE::NNE::IModelRDG>();
+		return TSharedPtr<UE::NNE::IModelRDG>();
 	}
 
 	TConstArrayView<uint8> Data = ModelData->GetModelData(GetRuntimeName());
@@ -147,5 +147,5 @@ TUniquePtr<UE::NNE::IModelRDG> UNNERuntimeRDGHlslImpl::CreateModel(TObjectPtr<UN
 		FEngineAnalytics::GetProvider().RecordEvent(TEXT("NeuralNetworkEngine.CreateModel"), Attributes);
 	}
 
-	return TUniquePtr<UE::NNE::IModelRDG>(Model);
+	return TSharedPtr<UE::NNE::IModelRDG>(Model);
 }
diff --git a/Engine/Plugins/Experimental/NNERuntimeRDG/Source/NNERuntimeRDG/Private/NNERuntimeRDGHlsl.h b/Engine/Plugins/Experimental/NNERuntimeRDG/Source/NNERuntimeRDG/Private/NNERuntimeRDGHlsl.h
--- a/Engine/Plugins/Experimental/NNERuntimeRDG/Source/NNERuntimeRDG/Private/NNERuntimeRDGHlsl.h
+++ b/Engine/Plugins/Experimental/NNERuntimeRDG/Source/NNERuntimeRDG/Private/NNERuntimeRDGHlsl.h
@@ -34,5 +34,5 @@ public:
 	virtual FString GetModelDataIdentifier(FString FileType, TConstArrayView<uint8> FileData, FGuid FileId, const ITargetPlatform* TargetPlatform) override;
 
 	virtual bool CanCreateModelRDG(TObjectPtr<UNNEModelData> ModelData) const override;
-	virtual TUniquePtr<UE::NNE::IModelRDG> CreateModel(TObjectPtr<UNNEModelData> ModelData) override;
+	virtual TSharedPtr<UE::NNE::IModelRDG> CreateModelRDG(TObjectPtr<UNNEModelData> ModelData) override;
 };
\ No newline at end of file
diff --git a/Engine/Plugins/Experimental/NNERuntimeRDG/Source/NNERuntimeRDG/Private/NNERuntimeRDGModelHlsl.cpp b/Engine/Plugins/Experimental/NNERuntimeRDG/Source/NNERuntimeRDG/Private/NNERuntimeRDGModelHlsl.cpp
--- a/Engine/Plugins/Experimental/NNERuntimeRDG/Source/NNERuntimeRDG/Private/NNERuntimeRDGModelHlsl.cpp
+++ b/Engine/Plugins/Experimental/NNERuntimeRDG/Source/NNERuntimeRDG/Private/NNERuntimeRDGModelHlsl.cpp
@@ -364,18 +364,18 @@ bool FModelInstance::PrepareWeights()
 	return true;
 }
 
-TUniquePtr<UE::NNE::IModelInstanceRDG> FModel::CreateModelInstance()
+TSharedPtr<UE::NNE::IModelInstanceRDG> FModel::CreateModelInstanceRDG()
 {
 	FModelInstance* ModelInstance = new FModelInstance();
 
 	if (!ModelInstance->Init(ModelData))
 	{
 		delete ModelInstance;
-		return TUniquePtr<UE::NNE::IModelInstanceRDG>();
+		return TSharedPtr<UE::NNE::IModelInstanceRDG>();
 	}
 
 	UE::NNE::IModelInstanceRDG* IModelInstance = static_cast<UE::NNE::IModelInstanceRDG*>(ModelInstance);
-	return TUniquePtr<UE::NNE::IModelInstanceRDG>(IModelInstance);
+	return TSharedPtr<UE::NNE::IModelInstanceRDG>(IModelInstance);
 }
 
 FModel::FModel(TConstArrayView<uint8> InModelData) : ModelData(InModelData) {}
diff --git a/Engine/Plugins/Experimental/NNERuntimeRDG/Source/NNERuntimeRDG/Private/NNERuntimeRDGModelHlsl.h b/Engine/Plugins/Experimental/NNERuntimeRDG/Source/NNERuntimeRDG/Private/NNERuntimeRDGModelHlsl.h
--- a/Engine/Plugins/Experimental/NNERuntimeRDG/Source/NNERuntimeRDG/Private/NNERuntimeRDGModelHlsl.h
+++ b/Engine/Plugins/Experimental/NNERuntimeRDG/Source/NNERuntimeRDG/Private/NNERuntimeRDGModelHlsl.h
@@ -41,7 +41,7 @@ public:
 	FModel(TConstArrayView<uint8> ModelData);
 	virtual ~FModel() {};
 
-	virtual TUniquePtr<UE::NNE::IModelInstanceRDG> CreateModelInstance() override;
+	virtual TSharedPtr<UE::NNE::IModelInstanceRDG> CreateModelInstanceRDG() override;
 
 private:
 	TArray<uint8> ModelData;
diff --git a/Engine/Plugins/Experimental/NNERuntimeRDG/Source/NNERuntimeRDG/Private/NNERuntimeRDGModule.cpp b/Engine/Plugins/Experimental/NNERuntimeRDG/Source/NNERuntimeRDG/Private/NNERuntimeRDGModule.cpp
--- a/Engine/Plugins/Experimental/NNERuntimeRDG/Source/NNERuntimeRDG/Private/NNERuntimeRDGModule.cpp
+++ b/Engine/Plugins/Experimental/NNERuntimeRDG/Source/NNERuntimeRDG/Private/NNERuntimeRDGModule.cpp
@@ -9,6 +9,10 @@
 
 void FNNERuntimeRDGModule::StartupModule()
 {
+#ifdef WITH_EDITOR
+	FModuleManager::Get().LoadModule(TEXT("NNEEditor"));
+#endif
+
 	// NNE runtime ORT Cpu startup
 	NNERuntimeRDGHlsl = NewObject<UNNERuntimeRDGHlslImpl>();
 	if (NNERuntimeRDGHlsl.IsValid())
diff --git a/Engine/Source/Editor/NNEEditor/NNEEditor.Build.cs b/Engine/Source/Editor/NNEEditor/NNEEditor.Build.cs
new file mode 100644
--- /dev/null
+++ b/Engine/Source/Editor/NNEEditor/NNEEditor.Build.cs
@@ -0,0 +1,29 @@
+// Copyright Epic Games, Inc. All Rights Reserved.
+
+using UnrealBuildTool;
+
+public class NNEEditor : ModuleRules
+{
+	public NNEEditor(ReadOnlyTargetRules Target) : base(Target)
+	{
+		PublicDependencyModuleNames.AddRange
+			(
+			new string[] {
+				"Core",
+				"CoreUObject",
+				"Engine"
+			}
+		);
+
+		PrivateDependencyModuleNames.AddRange(
+			new string[]
+			{
+				"Slate",
+				"SlateCore",
+				"UnrealEd",
+				"AssetTools",
+				"NNE"
+			}
+		);
+	}
+}
diff --git a/Engine/Source/Editor/NNEEditor/Private/NNEEditorModelDataActions.cpp b/Engine/Source/Editor/NNEEditor/Private/NNEEditorModelDataActions.cpp
new file mode 100644
--- /dev/null
+++ b/Engine/Source/Editor/NNEEditor/Private/NNEEditorModelDataActions.cpp
@@ -0,0 +1,60 @@
+// Copyright Epic Games, Inc. All Rights Reserved.
+
+#include "NNEEditorModelDataActions.h"
+
+#include "EditorFramework/AssetImportData.h"
+#include "NNEEditorModelDataEditorToolkit.h"
+#include "NNEModelData.h"
+
+#define LOCTEXT_NAMESPACE "AssetTypeActions"
+
+namespace UE::NNEEditor::Private
+{
+
+	UClass* FModelDataAssetTypeActions::GetSupportedClass() const
+	{
+		return UNNEModelData::StaticClass();
+	}
+
+	FText FModelDataAssetTypeActions::GetName() const
+	{
+		return LOCTEXT("FModelDataAssetTypeActionsName", "NNE Model Data");
+	}
+
+	FColor FModelDataAssetTypeActions::GetTypeColor() const
+	{
+		return FColor::Cyan;
+	}
+
+	uint32 FModelDataAssetTypeActions::GetCategories()
+	{
+		return EAssetTypeCategories::Misc;
+	}
+
+	void FModelDataAssetTypeActions::OpenAssetEditor(const TArray<UObject*>& InObjects, TSharedPtr<class IToolkitHost> EditWithinLevelEditor)
+	{
+		for (UObject* Obj : InObjects)
+		{
+			UNNEModelData* ModelData = Cast<UNNEModelData>(Obj);
+			if (ModelData != nullptr)
+			{
+				MakeShared<FModelDataEditorToolkit>()->InitEditor(ModelData);
+			}
+		}
+	}
+
+	void FModelDataAssetTypeActions::GetResolvedSourceFilePaths(const TArray<UObject*>& TypeAssets, TArray<FString>& OutSourceFilePaths) const
+	{
+		for (UObject* Asset : TypeAssets)
+		{
+			const UNNEModelData* ModelData = CastChecked<UNNEModelData>(Asset);
+			if (ModelData->AssetImportData)
+			{
+				ModelData->AssetImportData->ExtractFilenames(OutSourceFilePaths);
+			}
+		}
+	}
+
+} // UE::NNEEditor::Private
+
+#undef LOCTEXT_NAMESPACE
\ No newline at end of file
diff --git a/Engine/Source/Editor/NNEEditor/Private/NNEEditorModelDataActions.h b/Engine/Source/Editor/NNEEditor/Private/NNEEditorModelDataActions.h
new file mode 100644
--- /dev/null
+++ b/Engine/Source/Editor/NNEEditor/Private/NNEEditorModelDataActions.h
@@ -0,0 +1,26 @@
+// Copyright Epic Games, Inc. All Rights Reserved.
+
+#pragma once
+
+#include "CoreMinimal.h"
+#include "AssetTypeActions_Base.h"
+
+class UNNEModelData;
+
+namespace UE::NNEEditor::Private
+{
+	class FModelDataAssetTypeActions : public FAssetTypeActions_Base
+	{
+	public:
+		// IAssetTypeActions interface
+		virtual FText GetName() const override;
+		virtual FColor GetTypeColor() const override;
+		virtual UClass* GetSupportedClass() const override;
+		virtual uint32 GetCategories() override;
+		virtual bool IsImportedAsset() const override { return true; }
+		virtual void GetResolvedSourceFilePaths(const TArray<UObject*>& TypeAssets, TArray<FString>& OutSourceFilePaths) const override;
+		virtual void OpenAssetEditor(const TArray<UObject*>& InObjects, TSharedPtr<class IToolkitHost> EditWithinLevelEditor) override;
+		// End of IAssetTypeActions interface
+	};
+
+} // UE::NNEEditor::Private
\ No newline at end of file
diff --git a/Engine/Source/Editor/NNEEditor/Private/NNEEditorModelDataEditorToolkit.cpp b/Engine/Source/Editor/NNEEditor/Private/NNEEditorModelDataEditorToolkit.cpp
new file mode 100644
--- /dev/null
+++ b/Engine/Source/Editor/NNEEditor/Private/NNEEditorModelDataEditorToolkit.cpp
@@ -0,0 +1,70 @@
+// Copyright Epic Games, Inc. All Rights Reserved.
+
+#include "NNEEditorModelDataEditorToolkit.h"
+
+#include "NNE.h"
+#include "NNEModelData.h"
+#include "NNEEditorRuntimesWidget.h"
+#include "Widgets/Docking/SDockTab.h"
+
+namespace UE::NNEEditor::Private
+{
+
+	void FModelDataEditorToolkit::InitEditor(UNNEModelData* InModelData)
+	{
+		ModelData = InModelData;
+
+		const TSharedRef<FTabManager::FLayout> Layout = FTabManager::NewLayout("NNEModelDataEditorLayout")
+			->AddArea
+			(
+				FTabManager::NewPrimaryArea()->SetOrientation(Orient_Vertical)
+				->Split
+				(
+					FTabManager::NewStack()
+					->SetSizeCoefficient(0.8f)
+					->AddTab(TEXT("NNEModelDataTab"), ETabState::OpenedTab)
+				)
+			);
+
+		FAssetEditorToolkit::InitAssetEditor(EToolkitMode::Standalone, {}, TEXT("NNEModelDataEditor"), Layout, true, true, ModelData);
+	}
+
+	void FModelDataEditorToolkit::RegisterTabSpawners(const TSharedRef<class FTabManager>& InTabManager)
+	{
+		FAssetEditorToolkit::RegisterTabSpawners(InTabManager);
+
+		WorkspaceMenuCategory = InTabManager->AddLocalWorkspaceMenuCategory(INVTEXT("NNE Model Data Editor"));
+
+		InTabManager->RegisterTabSpawner("NNEModelDataTab", FOnSpawnTab::CreateLambda([this](const FSpawnTabArgs&)
+			{
+				return SNew(SDockTab)
+				[
+					SNew(SRuntimesWidget)
+					.TargetRuntimes(this, &FModelDataEditorToolkit::GetTargetRuntimes)
+					.OnTargetRuntimesChanged(this, &FModelDataEditorToolkit::SetTargetRuntimes)
+				];
+			}))
+			.SetDisplayName(INVTEXT("NNE Model Data"))
+			.SetGroup(WorkspaceMenuCategory.ToSharedRef());
+	}
+
+	void FModelDataEditorToolkit::UnregisterTabSpawners(const TSharedRef<class FTabManager>& InTabManager)
+	{
+		FAssetEditorToolkit::UnregisterTabSpawners(InTabManager);
+		InTabManager->UnregisterTabSpawner("NNEModelDataTab");
+	}
+
+	TArrayView<const FString> FModelDataEditorToolkit::GetTargetRuntimes() const
+	{
+		check(ModelData != nullptr);
+		return ModelData->GetTargetRuntimes();
+	}
+
+	void FModelDataEditorToolkit::SetTargetRuntimes(TArrayView<const FString> TargetRuntimes)
+	{
+		check(ModelData != nullptr);
+		ModelData->Modify();
+		ModelData->SetTargetRuntimes(TargetRuntimes);
+	}
+
+} // UE::NNEEditor::Private
\ No newline at end of file
diff --git a/Engine/Source/Editor/NNEEditor/Private/NNEEditorModelDataEditorToolkit.h b/Engine/Source/Editor/NNEEditor/Private/NNEEditorModelDataEditorToolkit.h
new file mode 100644
--- /dev/null
+++ b/Engine/Source/Editor/NNEEditor/Private/NNEEditorModelDataEditorToolkit.h
@@ -0,0 +1,34 @@
+// Copyright Epic Games, Inc. All Rights Reserved.
+
+#pragma once
+
+#include "CoreMinimal.h"
+#include "Toolkits/AssetEditorToolkit.h"
+
+class UNNEModelData;
+
+namespace UE::NNEEditor::Private
+{
+	class FModelDataEditorToolkit : public FAssetEditorToolkit
+	{
+	public:
+		void InitEditor(UNNEModelData* ModelData);
+
+		// FAssetEditorToolkit interface
+		virtual void RegisterTabSpawners(const TSharedRef<class FTabManager>& TabManager) override;
+		virtual void UnregisterTabSpawners(const TSharedRef<class FTabManager>& TabManager) override;
+		virtual FName GetToolkitFName() const override { return FName("NNEModelDataEditor"); }
+		virtual FText GetBaseToolkitName() const override { return INVTEXT("NNE Model Data Editor"); }
+		virtual FString GetWorldCentricTabPrefix() const override { return TEXT("NNE Model Data Editor "); }
+		virtual FLinearColor GetWorldCentricTabColorScale() const override { return {}; }
+		// End of FAssetEditorToolkit interface
+
+	private:
+		TArrayView<const FString> GetTargetRuntimes() const;
+		void SetTargetRuntimes(TArrayView<const FString> TargetRuntimes);
+
+		TArray<FString> RegisteredRuntimes;
+		UNNEModelData* ModelData;
+	};
+
+} // UE::NNEEditor::Private
\ No newline at end of file
diff --git a/Engine/Source/Editor/NNEEditor/Private/NNEEditorModelDataFactory.cpp b/Engine/Source/Editor/NNEEditor/Private/NNEEditorModelDataFactory.cpp
new file mode 100644
--- /dev/null
+++ b/Engine/Source/Editor/NNEEditor/Private/NNEEditorModelDataFactory.cpp
@@ -0,0 +1,136 @@
+// Copyright Epic Games, Inc. All Rights Reserved.
+
+#include "NNEEditorModelDataFactory.h"
+
+#include "Editor.h"
+#include "Editor/EditorEngine.h"
+#include "EditorFramework/AssetImportData.h"
+#include "EngineAnalytics.h"
+#include "HAL/FileManager.h"
+#include "Kismet/GameplayStatics.h"
+#include "NNE.h"
+#include "NNEModelData.h"
+#include "Subsystems/ImportSubsystem.h"
+
+UNNEModelDataFactory::UNNEModelDataFactory(const FObjectInitializer& ObjectInitializer) : UFactory(ObjectInitializer)
+{
+	bCreateNew = false;
+	bEditorImport = true;
+	SupportedClass = UNNEModelData::StaticClass();
+	ImportPriority = DefaultImportPriority;
+	Formats.Add("onnx;Open Neural Network Exchange Format");
+}
+
+UObject* UNNEModelDataFactory::FactoryCreateBinary(UClass* Class, UObject* InParent, FName Name, EObjectFlags Flags, UObject* Context, const TCHAR * Type, const uint8 *& Buffer, const uint8 * BufferEnd, FFeedbackContext* Warn)
+{
+	GEditor->GetEditorSubsystem<UImportSubsystem>()->BroadcastAssetPreImport(this, Class, InParent, Name, Type);
+
+	if (!Type || !Buffer || !BufferEnd || BufferEnd - Buffer <= 0)
+	{
+		GEditor->GetEditorSubsystem<UImportSubsystem>()->BroadcastAssetPostImport(this, nullptr);
+		return nullptr;
+	}
+
+	UNNEModelData* ModelData = NewObject<UNNEModelData>(InParent, Class, Name, Flags);
+	TConstArrayView<uint8> BufferView = MakeArrayView(Buffer, BufferEnd-Buffer);
+	ModelData->Init(Type, BufferView);
+
+	ModelData->AssetImportData->Update(GetCurrentFilename());
+	GEditor->GetEditorSubsystem<UImportSubsystem>()->BroadcastAssetPostImport(this, ModelData);
+
+	if (FEngineAnalytics::IsAvailable())
+	{
+		TArray<FAnalyticsEventAttribute> Attributes = MakeAnalyticsEventAttributeArray(
+			TEXT("PlatformName"), UGameplayStatics::GetPlatformName(),
+			TEXT("FactoryName"), TEXT("UNNEModelDataFactory"),
+			TEXT("ModelFileSize"), BufferView.Num()
+		);
+		FEngineAnalytics::GetProvider().RecordEvent(TEXT("NeuralNetworkEngine.FactoryCreateBinary"), Attributes);
+	}
+
+	return ModelData;
+}
+
+bool UNNEModelDataFactory::FactoryCanImport(const FString & Filename)
+{
+	return Filename.EndsWith(FString("onnx"));
+}
+
+bool UNNEModelDataFactory::CanReimport(UObject* Obj, TArray<FString>& OutFilenames)
+{
+	UNNEModelData* pModel = Cast<UNNEModelData>(Obj);
+	if (pModel)
+	{
+		pModel->AssetImportData->ExtractFilenames(OutFilenames);
+		return true;
+	}
+	return false;
+}
+
+void UNNEModelDataFactory::SetReimportPaths(UObject* Obj, const TArray<FString>& NewReimportPaths)
+{
+	UNNEModelData* pModel = Cast<UNNEModelData>(Obj);
+	if (pModel && ensure(NewReimportPaths.Num() == 1))
+	{
+		pModel->AssetImportData->UpdateFilenameOnly(NewReimportPaths[0]);
+	}
+}
+
+EReimportResult::Type UNNEModelDataFactory::Reimport(UObject* Obj)
+{
+	if (!Obj || !Obj->IsA(UNNEModelData::StaticClass()))
+	{
+		return EReimportResult::Failed;
+	}
+
+	UNNEModelData* pModel = Cast<UNNEModelData>(Obj);
+	const FString ResolvedSourceFilePath = pModel->AssetImportData->GetFirstFilename();
+	if (!ResolvedSourceFilePath.Len())
+	{
+		UE_LOG(LogNNE, Warning, TEXT("Cannot reimport NNE Model Data: it does not have a path stored."));
+		return EReimportResult::Failed;
+	}
+
+	// Ensure that the file provided by the path exists
+	if (IFileManager::Get().FileSize(*ResolvedSourceFilePath) == INDEX_NONE)
+	{
+		UE_LOG(LogNNE, Warning, TEXT("Cannot reimport NNE Model Data: source file cannot be found."));
+		return EReimportResult::Failed;
+	}
+
+	// Keep runtime list to be reapplied upon asset import.
+	TArray<FString, TInlineAllocator<10>> TargetRuntimes;
+	TargetRuntimes = pModel->GetTargetRuntimes();
+
+	bool OutCanceled = false;
+
+	if (ImportObject(pModel->GetClass(), pModel->GetOuter(), *pModel->GetName(), RF_Public | RF_Standalone, ResolvedSourceFilePath, nullptr, OutCanceled) != nullptr)
+	{
+		UE_LOG(LogNNE, Log, TEXT("Reimported successfully"));
+
+		pModel->SetTargetRuntimes(TargetRuntimes);
+		pModel->AssetImportData->Update(ResolvedSourceFilePath);
+
+		// Try to find the outer package so we can dirty it up
+		if (pModel->GetOuter())
+		{
+			pModel->GetOuter()->MarkPackageDirty();
+		}
+		else
+		{
+			pModel->MarkPackageDirty();
+		}
+	}
+	else if (OutCanceled)
+	{
+		UE_LOG(LogNNE, Warning, TEXT("Reimport canceled"));
+		return EReimportResult::Cancelled;
+	}
+	else
+	{
+		UE_LOG(LogNNE, Warning, TEXT("Reimport failed"));
+		return EReimportResult::Failed;
+	}
+
+	return EReimportResult::Succeeded;
+}
\ No newline at end of file
diff --git a/Engine/Source/Editor/NNEEditor/Private/NNEEditorModelDataFactory.h b/Engine/Source/Editor/NNEEditor/Private/NNEEditorModelDataFactory.h
new file mode 100644
--- /dev/null
+++ b/Engine/Source/Editor/NNEEditor/Private/NNEEditorModelDataFactory.h
@@ -0,0 +1,35 @@
+// Copyright Epic Games, Inc. All Rights Reserved.
+
+#pragma once
+
+#include "CoreMinimal.h"
+#include "EditorReimportHandler.h"
+#include "Factories/Factory.h"
+
+#include "NNEEditorModelDataFactory.generated.h"
+
+/**
+ * Class for importing and creating new UNNEModelData assets.
+ *
+ * The currently supported format is .onnx. The factory is invoked when a file of this format is dragged to the editor's content browser.
+ */
+UCLASS()
+class NNEEDITOR_API UNNEModelDataFactory : public UFactory, public FReimportHandler
+{
+	GENERATED_BODY()
+
+public:
+	UNNEModelDataFactory(const FObjectInitializer& ObjectInitializer);
+
+public:
+	//~ Begin UFactory Interface
+	virtual UObject* FactoryCreateBinary(UClass* Class, UObject* InParent, FName Name, EObjectFlags Flags, UObject* Context, const TCHAR* Type, const uint8*& Buffer, const uint8* BufferEnd, FFeedbackContext* Warn) override;
+	virtual bool FactoryCanImport(const FString& Filename) override;
+	//~ End UFactory Interface
+
+	//~ Begin FReimportHandler Interface
+	virtual bool CanReimport(UObject* Obj, TArray<FString>& OutFilenames) override;
+	virtual void SetReimportPaths(UObject* Obj, const TArray<FString>& NewReimportPaths) override;
+	virtual EReimportResult::Type Reimport(UObject* Obj) override;
+	//~ End FReimportHandler Interface
+};
\ No newline at end of file
diff --git a/Engine/Source/Editor/NNEEditor/Private/NNEEditorModule.cpp b/Engine/Source/Editor/NNEEditor/Private/NNEEditorModule.cpp
new file mode 100644
--- /dev/null
+++ b/Engine/Source/Editor/NNEEditor/Private/NNEEditorModule.cpp
@@ -0,0 +1,27 @@
+// Copyright Epic Games, Inc. All Rights Reserved.
+
+#include "Modules/ModuleManager.h"
+#include "NNEEditorModelDataActions.h"
+
+class FNNEEditorModule : public IModuleInterface
+{
+public:
+	virtual void StartupModule() override
+	{
+		ModelDataAssetTypeActions = MakeShared<UE::NNEEditor::Private::FModelDataAssetTypeActions>();
+		FAssetToolsModule::GetModule().Get().RegisterAssetTypeActions(ModelDataAssetTypeActions.ToSharedRef());
+	}
+
+	virtual void ShutdownModule() override
+	{
+		if (FModuleManager::Get().IsModuleLoaded("AssetTools"))
+		{
+			FAssetToolsModule::GetModule().Get().UnregisterAssetTypeActions(ModelDataAssetTypeActions.ToSharedRef());
+		}
+	}
+
+private:
+	TSharedPtr<UE::NNEEditor::Private::FModelDataAssetTypeActions> ModelDataAssetTypeActions;
+};
+
+IMPLEMENT_MODULE(FNNEEditorModule, NNEEditor)
\ No newline at end of file
diff --git a/Engine/Source/Editor/NNEEditor/Private/NNEEditorRuntimesWidget.cpp b/Engine/Source/Editor/NNEEditor/Private/NNEEditorRuntimesWidget.cpp
new file mode 100644
--- /dev/null
+++ b/Engine/Source/Editor/NNEEditor/Private/NNEEditorRuntimesWidget.cpp
@@ -0,0 +1,145 @@
+// Copyright Epic Games, Inc. All Rights Reserved.
+
+#include "NNEEditorRuntimesWidget.h"
+
+#include "NNE.h"
+#include "SlateOptMacros.h"
+#include "Styling/AppStyle.h"
+#include "Widgets/Input/SCheckBox.h"
+#include "Widgets/Layout/SBox.h"
+#include "Widgets/Layout/SWrapBox.h"
+#include "Widgets/SBoxPanel.h"
+#include "Widgets/Text/STextBlock.h"
+
+#define LOCTEXT_NAMESPACE "SNNEEditorRuntimesWidget"
+
+namespace UE::NNEEditor::Private
+{
+
+BEGIN_SLATE_FUNCTION_BUILD_OPTIMIZATION
+
+void SRuntimesWidget::Construct(const FArguments& InArgs)
+{
+	TargetRuntimes = InArgs._TargetRuntimes;
+	OnTargetRuntimesChanged = InArgs._OnTargetRuntimesChanged;
+
+	TSharedPtr<SWrapBox> WrapBox;
+
+	ChildSlot
+	[
+		SNew(SBox)
+		[
+			SNew(SVerticalBox)
+			+
+			SVerticalBox::Slot()
+			.AutoHeight()
+			.HAlign(HAlign_Fill)
+			.VAlign(VAlign_Fill)
+			.Padding(4.0)
+			[
+				SNew(STextBlock).Text(FText::FromString("Target Runtimes"))
+				.Font(FAppStyle::GetFontStyle("PropertyWindow.BoldFont"))
+			]
+			+
+			SVerticalBox::Slot()
+			.AutoHeight()
+			.HAlign(HAlign_Fill)
+			.VAlign(VAlign_Fill)
+			.Padding(4.0)
+			[
+				SAssignNew(WrapBox, SWrapBox)
+				.PreferredWidth(300.0f)
+			]
+		]
+	];
+
+	InitCheckBoxNames();
+
+	for (int32 i = 0; i < CheckBoxNames.Num(); ++i)
+	{
+		WrapBox->AddSlot()
+		.Padding(5)
+		.VAlign(VAlign_Top)
+		[
+			SNew(SCheckBox)
+			.Style(FAppStyle::Get(), i==0?"RadioButton":"Checkbox")//"All" in first position as a RadioButton.
+			.IsChecked(this, &SRuntimesWidget::IsEntryChecked, i)
+			.OnCheckStateChanged(this, &SRuntimesWidget::OnEntryChanged, i)
+			[
+				SNew(STextBlock).Text(FText::FromString(CheckBoxNames[i]))
+				.Font(FAppStyle::GetFontStyle("PropertyWindow.NormalFont"))
+			]
+		];
+	}
+}
+
+END_SLATE_FUNCTION_BUILD_OPTIMIZATION
+
+void SRuntimesWidget::OnEntryChanged(ECheckBoxState CheckType, int32 Index)
+{
+	check(Index < CheckBoxNames.Num());
+	const FString& RuntimeName = CheckBoxNames[Index];
+
+	if (!OnTargetRuntimesChanged.IsBound())
+		return;
+
+	if (RuntimeName == TEXT("All"))
+	{
+		if (CheckType == ECheckBoxState::Checked)
+		{
+			TArray<FString> EmptyTargetRuntimes;
+			OnTargetRuntimesChanged.Execute(EmptyTargetRuntimes);
+		}
+	}
+	else
+	{
+		TArray<FString, TInlineAllocator<10>> NewTargetRuntimes;
+		NewTargetRuntimes = TargetRuntimes.Get();
+		if (CheckType == ECheckBoxState::Checked)
+		{
+			NewTargetRuntimes.AddUnique(RuntimeName);
+		}
+		else
+		{
+			NewTargetRuntimes.Remove(RuntimeName);
+		}
+		OnTargetRuntimesChanged.Execute(NewTargetRuntimes);
+	}
+}
+
+ECheckBoxState	SRuntimesWidget::IsEntryChecked(int32 Index) const
+{
+	check(Index < CheckBoxNames.Num());
+	const FString& CheckBoxName = CheckBoxNames[Index];
+
+	if (CheckBoxName == TEXT("All"))
+	{
+		return TargetRuntimes.Get().Num() == 0 ? ECheckBoxState::Checked : ECheckBoxState::Unchecked;
+	}
+	
+	return TargetRuntimes.Get().Contains(CheckBoxName) ? ECheckBoxState::Checked : ECheckBoxState::Unchecked;
+}
+
+void SRuntimesWidget::InitCheckBoxNames()
+{
+	CheckBoxNames.Reset();
+
+	CheckBoxNames.Add(TEXT("All"));
+
+	for (TWeakInterfacePtr<INNERuntime> Runtime : UE::NNE::GetAllRuntimes())
+	{
+		CheckBoxNames.Add(Runtime->GetRuntimeName());
+	}
+
+	//A runtime could possibly be in the target list but not registered to the editor
+	for (const FString& EnabledRuntimeName : TargetRuntimes.Get())
+	{
+		if (CheckBoxNames.Contains(EnabledRuntimeName))
+			continue;
+		CheckBoxNames.Add(EnabledRuntimeName);
+	}
+}
+
+} //UE::NNEEditor::Private
+
+#undef LOCTEXT_NAMESPACE
\ No newline at end of file
diff --git a/Engine/Source/Editor/NNEEditor/Private/NNEEditorRuntimesWidget.h b/Engine/Source/Editor/NNEEditor/Private/NNEEditorRuntimesWidget.h
new file mode 100644
--- /dev/null
+++ b/Engine/Source/Editor/NNEEditor/Private/NNEEditorRuntimesWidget.h
@@ -0,0 +1,38 @@
+// Copyright Epic Games, Inc. All Rights Reserved.
+
+#pragma once
+
+#include "CoreMinimal.h"
+#include "Styling/SlateTypes.h"
+#include "Widgets/SCompoundWidget.h"
+
+namespace UE::NNEEditor::Private
+{
+
+	class SRuntimesWidget : public SCompoundWidget
+	{
+
+	public:
+		DECLARE_DELEGATE_OneParam(FOnTargetRuntimesChanged, TArrayView<const FString> /*NewTargetRuntimes*/)
+
+		SLATE_BEGIN_ARGS(SRuntimesWidget)
+			{}
+			SLATE_ATTRIBUTE(TArrayView<const FString>, TargetRuntimes)
+			SLATE_EVENT(FOnTargetRuntimesChanged, OnTargetRuntimesChanged)
+		SLATE_END_ARGS()
+
+		void Construct(const FArguments& InArgs);
+
+	private:
+		void OnEntryChanged(ECheckBoxState CheckType, int32 Index);
+		ECheckBoxState IsEntryChecked(int32 Index) const;
+
+		void AddCheckBox(const FString& RuntimeName, int32 Index);
+		void InitCheckBoxNames();
+
+		TArray<FString> CheckBoxNames;
+		TAttribute<TArrayView<const FString>> TargetRuntimes;
+		FOnTargetRuntimesChanged OnTargetRuntimesChanged;
+	};
+
+} // UE::NNEEditor::Private
\ No newline at end of file
diff --git a/Engine/Source/Runtime/NNE/Internal/NNEAttributeDataType.h b/Engine/Source/Runtime/NNE/Internal/NNEAttributeDataType.h
new file mode 100644
--- /dev/null
+++ b/Engine/Source/Runtime/NNE/Internal/NNEAttributeDataType.h
@@ -0,0 +1,25 @@
+// Copyright Epic Games, Inc. All Rights Reserved.
+
+#pragma once
+
+#include "UObject/Class.h"
+#include "UObject/Object.h"
+
+#include "NNEAttributeDataType.generated.h"
+
+/**
+ * Attribute data types
+ * 
+ * Note: also extend NNEAttributeValueTraits.h for more type support
+ */
+UENUM()
+enum class ENNEAttributeDataType : uint8
+{
+	None,
+	Float,								//!< 32-bit floating number
+	FloatArray,							//!< TArray of 32-bit floating numbers
+	Int32,								//!< 32-bit signed integer
+	Int32Array,							//!< TArray of 32-bit signed integers
+	String,								//!< built-in FString
+	StringArray							//!< TArray of built-in FString
+};
\ No newline at end of file
diff --git a/Engine/Source/Runtime/NNE/Internal/NNEAttributeMap.h b/Engine/Source/Runtime/NNE/Internal/NNEAttributeMap.h
new file mode 100644
--- /dev/null
+++ b/Engine/Source/Runtime/NNE/Internal/NNEAttributeMap.h
@@ -0,0 +1,80 @@
+// Copyright Epic Games, Inc. All Rights Reserved.
+
+#pragma once
+
+#include "NNEAttributeValue.h"
+
+namespace UE::NNE
+{
+
+class FAttributeMap
+{
+public:
+	//Set attribute
+	void SetAttribute(const FString& Name, const FNNEAttributeValue& Value)
+	{
+		checkf(nullptr == Attributes.FindByPredicate([Name](const FEntry& Other) { return Other.Name == Name; }), TEXT("Attribute name should be unique"));
+		Attributes.Emplace(Name, Value);
+	}
+
+	template<typename T>
+	T GetValue(const FString& Name) const
+	{
+		const FNNEAttributeValue* Value = GetAttributeValue(Name);
+		checkf(Value != nullptr, TEXT("Required attribute %s not found"), *Name);
+		return Value->GetValue<T>();
+	}
+
+	//Query attributes
+	template<typename T>
+	T GetValueOrDefault(const FString& Name, T Default) const
+	{
+		const FNNEAttributeValue* Value = GetAttributeValue(Name);
+		return Value == nullptr ? Default : Value->GetValue<T>();
+	}
+	
+	const FNNEAttributeValue* GetAttributeValue(const FString& Name) const
+	{
+		const FEntry * entry = Attributes.FindByPredicate([Name](const FEntry& Other) { return Other.Name == Name; });
+		if (entry != nullptr)
+		{
+			return &entry->Value;
+		}
+		return nullptr;
+	}
+
+	//Iteration
+	int32 Num() const
+	{
+		return Attributes.Num();
+	}
+
+	const FString& GetName(int32 Idx) const
+	{
+		check(Idx < Num());
+		return Attributes[Idx].Name;
+	}
+
+	const FNNEAttributeValue& GetAttributeValue(int32 Idx) const
+	{
+		check(Idx < Num());
+		return Attributes[Idx].Value;
+	}
+
+private:
+
+	struct FEntry
+	{
+		FEntry(const FString& InName, const FNNEAttributeValue& InValue)
+			: Name(InName), Value(InValue)
+		{
+		}
+		
+		FString Name;
+		FNNEAttributeValue Value;
+	};
+
+	TArray<FEntry> Attributes;
+};
+
+} // namespace UE::NNE
\ No newline at end of file
diff --git a/Engine/Source/Runtime/NNE/Internal/NNEAttributeValue.h b/Engine/Source/Runtime/NNE/Internal/NNEAttributeValue.h
new file mode 100644
--- /dev/null
+++ b/Engine/Source/Runtime/NNE/Internal/NNEAttributeValue.h
@@ -0,0 +1,57 @@
+// Copyright Epic Games, Inc. All Rights Reserved.
+
+#pragma once
+
+#include "NNEAttributeValueTraits.h"
+#include "Serialization/MemoryWriter.h"
+#include "Serialization/MemoryReader.h"
+#include "UObject/Class.h"
+#include "UObject/Object.h"
+
+#include "NNEAttributeValue.generated.h"
+
+USTRUCT()
+struct FNNEAttributeValue
+{
+	GENERATED_USTRUCT_BODY()
+
+	FNNEAttributeValue()
+	{
+		Type = ENNEAttributeDataType::None;
+	}
+
+	template<typename T>
+	explicit FNNEAttributeValue(T InValue)
+	{
+		FMemoryWriter writer(Value, true);
+		writer << InValue;
+
+		Type = TNNEAttributeValueTraits<T>::GetType();
+	}
+
+	template<typename T>
+	T GetValue() const
+	{
+		check((Type == TNNEAttributeValueTraits<T>::GetType()));
+
+		T Result;
+
+		FMemoryReader Reader(Value, true);
+		Reader << Result;
+
+		return Result;
+	}
+
+	ENNEAttributeDataType GetType() const
+	{
+		return Type;
+	}
+	
+private:
+	
+	UPROPERTY(VisibleAnywhere, Category = "Neural Network Inference")
+	ENNEAttributeDataType	Type;
+	
+	UPROPERTY(VisibleAnywhere, Category = "Neural Network Inference")
+	TArray<uint8>			Value;
+};
\ No newline at end of file
diff --git a/Engine/Source/Runtime/NNE/Internal/NNEAttributeValueTraits.h b/Engine/Source/Runtime/NNE/Internal/NNEAttributeValueTraits.h
new file mode 100644
--- /dev/null
+++ b/Engine/Source/Runtime/NNE/Internal/NNEAttributeValueTraits.h
@@ -0,0 +1,45 @@
+// Copyright Epic Games, Inc. All Rights Reserved.
+
+#pragma once
+
+#include "NNEAttributeDataType.h"
+
+template<typename T> struct TNNEAttributeValueTraits
+{
+	static constexpr ENNEAttributeDataType GetType()
+	{
+		static_assert(!sizeof(T), "Attribute value trait must be specialized for this type.");
+		return ENNEAttributeDataType::None;
+	}
+};
+
+// Attribute specializations
+template<> struct TNNEAttributeValueTraits<float>
+{
+	static constexpr ENNEAttributeDataType GetType() { return ENNEAttributeDataType::Float; }
+};
+
+template<> struct TNNEAttributeValueTraits<int32>
+{
+	static constexpr ENNEAttributeDataType GetType() { return ENNEAttributeDataType::Int32; }
+};
+
+template<> struct TNNEAttributeValueTraits<FString>
+{
+	static constexpr ENNEAttributeDataType GetType() { return ENNEAttributeDataType::String; }
+};
+
+template<> struct TNNEAttributeValueTraits<TArray<int32>>
+{
+	static constexpr ENNEAttributeDataType GetType() { return ENNEAttributeDataType::Int32Array; }
+};
+
+template<> struct TNNEAttributeValueTraits<TArray<float>>
+{
+	static constexpr ENNEAttributeDataType GetType() { return ENNEAttributeDataType::FloatArray; }
+};
+
+template<> struct TNNEAttributeValueTraits<TArray<FString>>
+{
+	static constexpr ENNEAttributeDataType GetType() { return ENNEAttributeDataType::StringArray; }
+};
\ No newline at end of file
diff --git a/Engine/Source/Runtime/NNE/Internal/NNEModelBase.h b/Engine/Source/Runtime/NNE/Internal/NNEModelBase.h
new file mode 100644
--- /dev/null
+++ b/Engine/Source/Runtime/NNE/Internal/NNEModelBase.h
@@ -0,0 +1,80 @@
+// Copyright Epic Games, Inc. All Rights Reserved.
+
+#pragma once
+
+#include "CoreMinimal.h"
+#include "NNE.h"
+#include "NNETypes.h"
+
+namespace UE::NNE::Internal
+{
+	template <class ModelInterface> class FModelInstanceBase : public ModelInterface
+	{
+	public:
+
+		virtual ~FModelInstanceBase() = default;
+
+		virtual TConstArrayView<NNE::FTensorDesc> GetInputTensorDescs() const override;
+		virtual TConstArrayView<NNE::FTensorDesc> GetOutputTensorDescs() const override;
+		virtual TConstArrayView<NNE::FTensorShape> GetInputTensorShapes() const override;
+		virtual TConstArrayView<NNE::FTensorShape> GetOutputTensorShapes() const override;
+		virtual int SetInputTensorShapes(TConstArrayView<NNE::FTensorShape> InInputShapes) override;
+
+	protected:
+
+		FModelInstanceBase() {}
+
+		TArray<FTensorShape>	InputTensorShapes;
+		TArray<FTensorShape>	OutputTensorShapes;
+		TArray<FTensorDesc>		InputSymbolicTensors;
+		TArray<FTensorDesc>		OutputSymbolicTensors;
+	};
+
+	template <class T> TConstArrayView<FTensorDesc> FModelInstanceBase<T>::GetInputTensorDescs() const
+	{
+		return InputSymbolicTensors;
+	}
+
+	template <class T> TConstArrayView<FTensorDesc> FModelInstanceBase<T>::GetOutputTensorDescs() const
+	{
+		return OutputSymbolicTensors;
+	}
+
+	template <class T> TConstArrayView<FTensorShape> FModelInstanceBase<T>::GetInputTensorShapes() const
+	{
+		return InputTensorShapes;
+	}
+
+	template <class T> TConstArrayView<FTensorShape> FModelInstanceBase<T>::GetOutputTensorShapes() const
+	{
+		return OutputTensorShapes;
+	}
+
+	template <class T> int FModelInstanceBase<T>::SetInputTensorShapes(TConstArrayView<FTensorShape> InInputShapes)
+	{
+		InputTensorShapes.Reset(InInputShapes.Num());
+
+		if (InInputShapes.Num() != InputSymbolicTensors.Num())
+		{
+			UE_LOG(LogNNE, Warning, TEXT("Number of input shapes does not match number of input tensors"));
+			return -1;
+		}
+
+		for (int32 i = 0; i < InInputShapes.Num(); ++i)
+		{
+			const FTensorDesc SymbolicDesc = InputSymbolicTensors[i];
+			if (!InInputShapes[i].IsCompatibleWith(SymbolicDesc.GetShape()))
+			{
+				UE_LOG(LogNNE, Warning, TEXT("Input shape does not match input tensor %s of index %d"), *SymbolicDesc.GetName(), i);
+				return -1;
+			}
+		}
+
+		InputTensorShapes = InInputShapes;
+
+		//Implementations are responsible to handle output and intermediate tensor shape inference.
+		//This base implementation only validate that all inputs are matching what the model can support.
+		return 0;
+	}
+
+} // namespace UE::NNE::Internal
\ No newline at end of file
diff --git a/Engine/Source/Runtime/NNE/Internal/NNEModelOptimizerInterface.h b/Engine/Source/Runtime/NNE/Internal/NNEModelOptimizerInterface.h
new file mode 100644
--- /dev/null
+++ b/Engine/Source/Runtime/NNE/Internal/NNEModelOptimizerInterface.h
@@ -0,0 +1,55 @@
+// Copyright Epic Games, Inc. All Rights Reserved.
+
+#pragma once
+
+#include "Containers/UnrealString.h"
+#include "Templates/SharedPointerFwd.h"
+
+struct FNNEModelRaw;
+namespace UE::NNE { class FAttributeMap; }
+
+namespace UE::NNE::Internal
+{
+	
+//Note: create a dedicated FOptimizerOptionsMap when function diverge with FAttributeMap
+//example introduction of sparse tensor to FAttribute witch make no sense as an optimizer attribute
+using FOptimizerOptionsMap = UE::NNE::FAttributeMap;
+
+/** Interface class for NNE model validator */
+class IModelValidator
+{
+public:
+	virtual ~IModelValidator() = default;
+	virtual FString GetName() const = 0;
+	virtual bool ValidateModel(const FNNEModelRaw& InputModel, const FOptimizerOptionsMap& Options) const = 0;
+};
+
+/** Interface class for NNE model optimizer pass */
+class IModelOptimizerPass
+{
+public:
+	virtual ~IModelOptimizerPass() = default;
+	virtual FString GetName() const = 0;
+
+	//Optimize the model in place, potentially changing the format
+	virtual bool ApplyPass(FNNEModelRaw& Model, const FOptimizerOptionsMap& Options) const = 0;
+};
+
+/** Interface class for NNE model optimizer */
+class IModelOptimizer
+{
+public:
+	virtual ~IModelOptimizer() = default;
+	virtual FString GetName() const = 0;
+
+	//Allow to extend/customize an optimizer by adding passes. They should be executed in order.
+	virtual void AddOptimizationPass(TSharedPtr<IModelOptimizerPass> ModelOptimizerPass) = 0;
+	
+	//Allow to extend/customize an optimizer all validators should be run between each pass.
+	virtual void AddValidator(TSharedPtr<IModelValidator>) = 0;
+	
+	//Apply all passes and validators to the input model, produce an optimized model potentially in a different format
+	virtual bool Optimize(const FNNEModelRaw& InputModel, FNNEModelRaw& OutModel, const FOptimizerOptionsMap& Options) = 0;
+};
+
+} // namespace UE::NNE::Internal
diff --git a/Engine/Source/Runtime/NNE/Internal/NNEOperator.h b/Engine/Source/Runtime/NNE/Internal/NNEOperator.h
new file mode 100644
--- /dev/null
+++ b/Engine/Source/Runtime/NNE/Internal/NNEOperator.h
@@ -0,0 +1,89 @@
+// Copyright Epic Games, Inc. All Rights Reserved.
+
+#pragma once
+
+#include "HAL/Platform.h"
+
+namespace UE::NNE::Internal
+{
+
+	//One input element-wise operators
+	enum class EElementWiseUnaryOperatorType : uint8
+	{
+		Abs,
+		Acos,
+		Acosh,
+		Asin,
+		Asinh,
+		Atan,
+		Atanh,
+		//BitShift, //Note: need integer tensors
+		Ceil,
+		Clip,
+		Cos,
+		Cosh,
+		Elu,
+		Erf,
+		Exp,
+		Floor,
+		IsInf,
+		IsNan,
+		HardSigmoid,
+		HardSwish,
+		LeakyRelu,
+		Log,
+		Neg,
+		//Not,      //Note: need bool tensors
+		Reciprocal,
+		Relu,
+		Round,
+		Selu,
+		Sigmoid,
+		Sign,
+		Sin,
+		Sinh,
+		Softplus,
+		Softsign,
+		Sqrt,
+		Tan,
+		Tanh,
+
+		MAX
+	};
+
+	//Two inputs element-wise operators with multi-directional broadcast
+	//see [REDACTED_PROFILE_URL]/onnx/blob/main/docs/Broadcasting.md
+	enum class EElementWiseBinaryOperatorType : uint8
+	{
+		Add,
+		//And,           //Note: need boolean tensors
+		Div,
+		//Equal,         //Note: need boolean tensors
+		//Greater,       //Note: need boolean tensors
+		//GreaterOrEqual,//Note: need boolean tensors
+		//Less,          //Note: need boolean tensors
+		//LessOrEqual,   //Note: need boolean tensors
+		Mod,
+		Mul,
+		//Or,            //Note: need boolean tensors
+		Prelu,           //Note: only broadcast from slope to input.
+		Pow,
+		Sub,
+		//Xor,           //Note:  need boolean tensors
+
+		MAX
+	};
+
+	//Variable number of inputs element-wise operators with multi-directional broadcast
+	//see [REDACTED_PROFILE_URL]/onnx/blob/main/docs/Broadcasting.md
+	enum class EElementWiseVariadicOperatorType : uint8
+	{
+		Max,
+		Min,
+		Mean,
+		Sum,
+
+		MAX
+	};
+
+} // namespace UE::NNE::Internal
\ No newline at end of file
diff --git a/Engine/Source/Runtime/NNE/Internal/NNERuntimeFormat.h b/Engine/Source/Runtime/NNE/Internal/NNERuntimeFormat.h
new file mode 100644
--- /dev/null
+++ b/Engine/Source/Runtime/NNE/Internal/NNERuntimeFormat.h
@@ -0,0 +1,113 @@
+// Copyright Epic Games, Inc. All Rights Reserved.
+
+#pragma once
+
+#include "UObject/Object.h"
+#include "NNEAttributeValue.h"
+#include "NNETypes.h"
+
+#include "NNERuntimeFormat.generated.h"
+
+UENUM()
+enum class ENNEFormatTensorType : uint8
+{
+	None,	
+	Input,
+	Output,
+	Intermediate,
+	Initializer,
+
+	NUM
+};
+
+UENUM()
+enum class ENNEInferenceFormat : uint8
+{
+	Invalid,
+	ONNX,				//!< ONNX Open Neural Network Exchange
+	ORT,				//!< ONNX Runtime (only for CPU)
+	NNERT				//!< NNE Runtime format
+};
+
+USTRUCT()
+struct FNNEModelRaw
+{
+	GENERATED_USTRUCT_BODY()
+	
+	UPROPERTY(VisibleAnywhere, Category = "Neural Network Inference")
+	TArray<uint8>		Data;
+	
+	UPROPERTY(VisibleAnywhere, Category = "Neural Network Inference")
+	ENNEInferenceFormat	Format { ENNEInferenceFormat::Invalid };
+};
+
+// Required by LoadModel() when loading operators in HLSL and DirectML runtime
+USTRUCT()
+struct FNNEFormatAttributeDesc
+{
+	GENERATED_USTRUCT_BODY()
+
+	UPROPERTY(VisibleAnywhere, Category = "Neural Network Inference")
+	FString Name;
+
+	UPROPERTY(VisibleAnywhere, Category = "Neural Network Inference")
+	FNNEAttributeValue Value;
+};
+
+USTRUCT()
+struct FNNEFormatOperatorDesc
+{
+	GENERATED_USTRUCT_BODY()
+
+	UPROPERTY(VisibleAnywhere, Category = "Neural Network Inference")
+	FString TypeName;			//!< For example "Relu"
+
+	UPROPERTY(VisibleAnywhere, Category = "Neural Network Inference")
+	TArray<uint32> InTensors;
+
+	UPROPERTY(VisibleAnywhere, Category = "Neural Network Inference")
+	TArray<uint32> OutTensors;
+
+	UPROPERTY(VisibleAnywhere, Category = "Neural Network Inference")
+	TArray<FNNEFormatAttributeDesc> Attributes;
+};
+
+USTRUCT()
+struct FNNEFormatTensorDesc
+{
+	GENERATED_USTRUCT_BODY()
+
+	UPROPERTY(VisibleAnywhere, Category = "Neural Network Inference")
+	FString Name;
+
+	UPROPERTY(VisibleAnywhere, Category = "Neural Network Inference")
+	TArray<int32> Shape;
+
+	UPROPERTY(VisibleAnywhere, Category = "Neural Network Inference")
+	ENNEFormatTensorType	Type = ENNEFormatTensorType::None;
+
+	UPROPERTY(VisibleAnywhere, Category = "Neural Network Inference")
+	ENNETensorDataType	DataType = ENNETensorDataType::None;
+
+	UPROPERTY(VisibleAnywhere, Category = "Neural Network Inference")
+	uint64	DataSize = 0;
+
+	UPROPERTY(VisibleAnywhere, Category = "Neural Network Inference")
+	uint64	DataOffset = 0;
+};
+
+/// NNE Runtime format
+USTRUCT()
+struct FNNERuntimeFormat
+{
+	GENERATED_USTRUCT_BODY()
+
+	UPROPERTY(VisibleAnywhere, Category = "Neural Network Inference")
+	TArray<FNNEFormatTensorDesc> Tensors;
+
+	UPROPERTY(VisibleAnywhere, Category = "Neural Network Inference")
+	TArray<FNNEFormatOperatorDesc> Operators;
+
+	UPROPERTY(VisibleAnywhere, Category = "Neural Network Inference")
+	TArray<uint8> TensorData;
+};
\ No newline at end of file
diff --git a/Engine/Source/Runtime/NNE/Internal/NNETensor.h b/Engine/Source/Runtime/NNE/Internal/NNETensor.h
new file mode 100644
--- /dev/null
+++ b/Engine/Source/Runtime/NNE/Internal/NNETensor.h
@@ -0,0 +1,110 @@
+// Copyright Epic Games, Inc. All Rights Reserved.
+
+#pragma once
+
+#include "CoreMinimal.h"
+#include "NNETypes.h"
+
+namespace UE::NNE::Internal
+{
+	/** Concrete tensor with data accessible by graph scheduling */
+	class FTensor
+	{
+	protected:
+		FString				Name;
+		ENNETensorDataType	DataType;
+		FTensorShape		Shape;
+		TArray<uint8>		PreparedData;
+		uint64				DataSize = 0;
+		uint32				Volume = 0;
+
+		FTensor() = default;
+		
+	public:
+		const FString& GetName() const
+		{
+			return Name;
+		}
+
+		ENNETensorDataType GetDataType() const
+		{
+			return DataType;
+		}
+
+		uint32 GetElemByteSize() const
+		{
+			return GetTensorDataTypeSizeInBytes(DataType);
+		}
+		
+		const FTensorShape& GetShape() const
+		{
+			return Shape;
+		}
+
+		template <typename T> TConstArrayView<T> GetPreparedData() const
+		{
+			const T* DataPtr = reinterpret_cast<const T*>(PreparedData.GetData());
+			const int32 ElemSize = sizeof(T);
+
+			check(PreparedData.Num() % ElemSize == 0);
+			return MakeArrayView(DataPtr, PreparedData.Num() / ElemSize);
+		}
+
+		void SetShape(const FTensorShape& InShape)
+		{
+			checkf(!HasPreparedData(), TEXT("Shape cannot be changed once data as been set."));
+			check(InShape.Volume() <= TNumericLimits<uint32>::Max());
+			Shape = InShape;
+			Volume = InShape.Volume();
+			DataSize = static_cast<uint64>(GetTensorDataTypeSizeInBytes(DataType) * Volume);
+		}
+
+		template <typename T> void SetPreparedData(TConstArrayView<T> Data)
+		{
+			const uint8* DataPtr = reinterpret_cast<const uint8*>(Data.GetData());
+			TConstArrayView<uint8> DataAsByte = MakeArrayView(DataPtr, Data.Num() * sizeof(T));
+			
+			checkf(DataAsByte.Num() == DataSize, TEXT("Incorrect data size, it should match tensor shape and data type."));
+			PreparedData.Reset();
+			PreparedData.Append(DataAsByte);
+		}
+
+		bool HasPreparedData() const
+		{
+			return !PreparedData.IsEmpty();
+		}
+
+		uint32 GetVolume() const
+		{
+			return Volume;
+		}
+
+		uint64 GetDataSize() const
+		{
+			return DataSize;
+		}
+
+		static FTensor Make(const FString& Name, const FTensorShape& Shape, ENNETensorDataType DataType)
+		{
+			FTensor Tensor;
+			Tensor.Name = Name;
+			Tensor.DataType = DataType;
+			Tensor.SetShape(Shape);
+			return Tensor;
+		}
+
+		static FTensor Make(const FTensorDesc& TensorDesc, const FTensorShape& Shape)
+		{
+			check(Shape.IsCompatibleWith(TensorDesc.GetShape()));
+			return Make(TensorDesc.GetName(), Shape, TensorDesc.GetDataType());
+		}
+
+		static FTensor MakeFromSymbolicDesc(const FTensorDesc& TensorDesc)
+		{
+			return Make(TensorDesc.GetName(), FTensorShape::MakeFromSymbolic(TensorDesc.GetShape()), TensorDesc.GetDataType());
+		}
+	};
+
+	using FTensorRef = FTensor*;
+
+} // namespace UE::NNE::Internal
\ No newline at end of file
diff --git a/Engine/Source/Runtime/NNE/NNE.Build.cs b/Engine/Source/Runtime/NNE/NNE.Build.cs
new file mode 100644
--- /dev/null
+++ b/Engine/Source/Runtime/NNE/NNE.Build.cs
@@ -0,0 +1,29 @@
+// Copyright Epic Games, Inc. All Rights Reserved.
+
+namespace UnrealBuildTool.Rules
+{
+    public class NNE : ModuleRules
+    {
+        public NNE(ReadOnlyTargetRules Target) : base(Target)
+        {
+            PrivateDependencyModuleNames.AddRange(
+                new string[]
+				{
+					"Core",
+					"CoreUObject",
+					"Engine",
+				}
+			);
+
+			if (Target.Type == TargetType.Editor)
+			{
+				PrivateDependencyModuleNames.AddRange(
+					new string[]
+					{
+						"DerivedDataCache"
+					}
+				);
+			}
+		}
+    }
+}
diff --git a/Engine/Source/Runtime/NNE/Private/NNE.cpp b/Engine/Source/Runtime/NNE/Private/NNE.cpp
new file mode 100644
--- /dev/null
+++ b/Engine/Source/Runtime/NNE/Private/NNE.cpp
@@ -0,0 +1,121 @@
+// Copyright Epic Games, Inc. All Rights Reserved.
+
+#include "NNE.h"
+
+#include "EngineAnalytics.h"
+#include "Kismet/GameplayStatics.h"
+#include "Misc/CoreDelegates.h"
+#include "Misc/SecureHash.h"
+
+DEFINE_LOG_CATEGORY(LogNNE);
+
+namespace UE::NNE
+{
+	class FRegistry
+	{
+	public:
+
+		static FRegistry* Get()
+		{
+			static FRegistry Inst;
+
+			return &Inst;
+		}
+
+		bool Add(TWeakInterfacePtr<INNERuntime> Runtime)
+		{
+			if (!Runtime.IsValid())
+			{
+				return false;
+			}
+			
+			if (FindByName(Runtime->GetRuntimeName()) >= 0)
+			{
+				UE_LOG(LogNNE, Warning, TEXT("Runtime %s is already registered"), *Runtime->GetRuntimeName());
+				return false;
+			}
+
+			Runtimes.Add(Runtime);
+
+			return true;
+		}
+
+		bool Remove(TWeakInterfacePtr<INNERuntime> Runtime)
+		{
+			if (!Runtime.IsValid())
+			{
+				return false;
+			}
+			
+			int Index = FindByName(Runtime->GetRuntimeName());
+			if (Index >= 0)
+			{
+				Runtimes.RemoveAtSwap(Index);
+				return true;
+			}
+
+			return false;
+		}
+
+
+		TArrayView<TWeakInterfacePtr<INNERuntime>> GetAllRuntimes()
+		{
+			RemoveInvalidPtrs();
+			return Runtimes;
+		}
+
+	private:
+
+		void RemoveInvalidPtrs()
+		{
+			Runtimes.RemoveAllSwap([](TWeakInterfacePtr<INNERuntime> Runtime) -> bool {return !Runtime.IsValid();}, false);
+		}
+
+		int FindByName(const FString& Name)
+		{
+			RemoveInvalidPtrs();
+			
+			int Index = -1;
+
+			for (int Idx = 0; Idx < Runtimes.Num(); ++Idx)
+			{
+				if (Runtimes[Idx]->GetRuntimeName() == Name)
+				{
+					Index = Idx;
+					break;
+				}
+			}
+
+			return Index;
+		}
+
+		TArray<TWeakInterfacePtr<INNERuntime>>	Runtimes;
+	};
+
+	bool RegisterRuntime(TWeakInterfacePtr<INNERuntime> Runtime)
+	{
+		FString RuntimeName = Runtime->GetRuntimeName();
+		FCoreDelegates::OnAllModuleLoadingPhasesComplete.AddLambda([RuntimeName]() {
+			if (FEngineAnalytics::IsAvailable())
+			{
+				TArray<FAnalyticsEventAttribute> Attributes = MakeAnalyticsEventAttributeArray(
+					TEXT("PlatformName"), UGameplayStatics::GetPlatformName(),
+					TEXT("HashedRuntimeName"), FMD5::HashAnsiString(*RuntimeName)
+				);
+				FEngineAnalytics::GetProvider().RecordEvent(TEXT("NeuralNetworkEngine.RegisterRuntime"), Attributes);
+			}
+		});
+
+		return FRegistry::Get()->Add(Runtime);
+	}
+
+	bool UnregisterRuntime(TWeakInterfacePtr<INNERuntime> Runtime)
+	{
+		return FRegistry::Get()->Remove(Runtime);
+	}
+
+	TArrayView<TWeakInterfacePtr<INNERuntime>> GetAllRuntimes()
+	{
+		return FRegistry::Get()->GetAllRuntimes();
+	}
+}
diff --git a/Engine/Source/Runtime/NNE/Private/NNEModelData.cpp b/Engine/Source/Runtime/NNE/Private/NNEModelData.cpp
new file mode 100644
--- /dev/null
+++ b/Engine/Source/Runtime/NNE/Private/NNEModelData.cpp
@@ -0,0 +1,363 @@
+// Copyright Epic Games, Inc. All Rights Reserved.
+
+#include "NNEModelData.h"
+
+#include "NNE.h"
+#include "NNEAttributeMap.h"
+#include "NNEModelOptimizerInterface.h"
+#include "NNERuntimeFormat.h"
+#include "Serialization/CustomVersion.h"
+#include "UObject/WeakInterfacePtr.h"
+#include "EditorFramework/AssetImportData.h"
+
+#if WITH_EDITOR
+#include "Containers/StringFwd.h"
+#include "DerivedDataCacheKey.h"
+#include "DerivedDataCache.h"
+#include "DerivedDataRequestOwner.h"
+#include "Internationalization/TextLocalizationResource.h"
+#include "Memory/SharedBuffer.h"
+#include "Misc/Guid.h"
+#endif
+
+enum Type
+{
+	Initial = 0,
+	TargetRuntimesAndAssetImportData = 1,
+	// -----<new versions can be added before this line>-------------------------------------------------
+	// - this needs to be the last line (see note below)
+	VersionPlusOne,
+	LatestVersion = VersionPlusOne - 1
+};
+
+const FGuid UNNEModelData::GUID(0x9513202e, 0xeba1b279, 0xf17fe5ba, 0xab90c3f2);
+FCustomVersionRegistration NNEModelDataVersion(UNNEModelData::GUID, LatestVersion, TEXT("NNEModelDataVersion"));//Always save with the latest version
+
+#if WITH_EDITOR
+
+inline FString GetDDCRequestId(const FString& FileId, const FString& RuntimeName, const FString& ModelDataIdentifier)
+{
+	//RuntimeName and FileId are embedded to the id to ensure no potential collision between the runtime/assets
+	return RuntimeName + "-" + FileId + "-" + ModelDataIdentifier;
+}
+
+inline UE::DerivedData::FCacheKey CreateCacheKey(const FString& FileId, const FString& RequestId)
+{
+	return { UE::DerivedData::FCacheBucket(FWideStringView(*FileId)), FIoHash::HashBuffer(MakeMemoryView(FTCHARToUTF8(RequestId))) };
+}
+
+inline FSharedBuffer GetFromDDC(const FGuid& FileId, const FString& RuntimeName, const FString& ModelDataIdentifier)
+{
+	FString FileIdStr = FileId.ToString(EGuidFormats::Digits);
+	FString RequestId = GetDDCRequestId(FileIdStr, RuntimeName, ModelDataIdentifier);
+
+	UE::DerivedData::FCacheGetValueRequest GetRequest;
+	GetRequest.Name = FString("Get-") + RequestId;
+	GetRequest.Key = CreateCacheKey(FileIdStr, RequestId);
+	FSharedBuffer RawDerivedData;
+	UE::DerivedData::FRequestOwner BlockingGetOwner(UE::DerivedData::EPriority::Blocking);
+	UE::DerivedData::GetCache().GetValue({ GetRequest }, BlockingGetOwner, [&RawDerivedData](UE::DerivedData::FCacheGetValueResponse&& Response)
+		{
+			RawDerivedData = Response.Value.GetData().Decompress();
+		});
+	BlockingGetOwner.Wait();
+	return RawDerivedData;
+}
+
+inline void PutIntoDDC(const FGuid& FileId, const FString& RuntimeName, const FString& ModelDataIdentifier, FSharedBuffer& Data)
+{
+	FString FileIdStr = FileId.ToString(EGuidFormats::Digits);
+	FString RequestId = GetDDCRequestId(FileIdStr, RuntimeName, ModelDataIdentifier);
+
+	UE::DerivedData::FCachePutValueRequest PutRequest;
+	PutRequest.Name = FString("Put-") + RequestId;
+	PutRequest.Key = CreateCacheKey(FileIdStr, RequestId);
+	PutRequest.Value = UE::DerivedData::FValue::Compress(Data);
+	UE::DerivedData::FRequestOwner BlockingPutOwner(UE::DerivedData::EPriority::Blocking);
+	UE::DerivedData::GetCache().PutValue({ PutRequest }, BlockingPutOwner);
+	BlockingPutOwner.Wait();
+}
+
+#endif
+
+inline TArray<uint8> CreateRuntimeDataBlob(const FString& RuntimeName, FString FileType, const TArray<uint8>& FileData, FGuid FileId, const ITargetPlatform* TargetPlatform)
+{
+	TWeakInterfacePtr<INNERuntime> NNERuntime = UE::NNE::GetRuntime<INNERuntime>(RuntimeName);
+	if (NNERuntime.IsValid())
+	{
+		return NNERuntime->CreateModelData(FileType, FileData, FileId, TargetPlatform);
+	}
+	else
+	{
+		UE_LOG(LogNNE, Error, TEXT("UNNEModelData: No runtime '%s' found. Valid runtimes are: "), *RuntimeName);
+		TArrayView<TWeakInterfacePtr<INNERuntime>> Runtimes = UE::NNE::GetAllRuntimes();
+		for (int i = 0; i < Runtimes.Num(); i++)
+		{
+			UE_LOG(LogNNE, Error, TEXT("- %s"), *Runtimes[i]->GetRuntimeName());
+		}
+		return {};
+	}
+}
+
+void UNNEModelData::Init(const FString& Type, TConstArrayView<uint8> Buffer)
+{
+	FileType = Type;
+	FileData = Buffer;
+	FPlatformMisc::CreateGuid(FileId);
+	ModelData.Empty();
+}
+
+FString UNNEModelData::GetFileType()
+{
+	return FileType;
+}
+
+TConstArrayView<uint8> UNNEModelData::GetFileData()
+{
+	return FileData;
+}
+
+FGuid UNNEModelData::GetFileId()
+{
+	return FileId;
+}
+
+TConstArrayView<uint8> UNNEModelData::GetModelData(const FString& RuntimeName)
+{
+#if WITH_EDITORONLY_DATA
+	// Check model data is supporting the requested target runtime
+	TArrayView<const FString> TargetRuntimesNames = GetTargetRuntimes();
+	if (!TargetRuntimesNames.IsEmpty() && !TargetRuntimesNames.Contains(RuntimeName))
+	{
+		UE_LOG(LogNNE, Error, TEXT("UNNEModelData: Runtime '%s' is not among the target runtimes. Target runtimes are: "), *RuntimeName);
+		for (const FString& TargetRuntimesName : TargetRuntimesNames)
+		{
+			UE_LOG(LogNNE, Error, TEXT("- %s"), *TargetRuntimesName);
+		}
+		return {};
+	}
+#endif //WITH_EDITORONLY_DATA
+
+	// Check if we have a local cache hit
+	TArray<uint8>* LocalData = ModelData.Find(RuntimeName);
+	if (LocalData)
+	{
+		return TConstArrayView<uint8>(LocalData->GetData(), LocalData->Num());
+	}
+	
+#if WITH_EDITOR
+	TWeakInterfacePtr<INNERuntime> NNERuntime = UE::NNE::GetRuntime<INNERuntime>(RuntimeName);
+	if (!NNERuntime.IsValid())
+	{
+		UE_LOG(LogNNE, Error, TEXT("UNNEModelData: Runtime '%s' is among the target runtimes but instance is invalid."), *RuntimeName);
+		return {};
+	}
+
+	FString ModelDataIdentifier = NNERuntime->GetModelDataIdentifier(FileType, FileData, FileId, nullptr);
+	if (ModelDataIdentifier.Len() == 0)
+	{
+		UE_LOG(LogNNE, Error, TEXT("UNNEModelData: Runtime '%s' returned an empty string as a ModelDataIdentifier. GetModelDataIdentifier should always return a valid identifier."), *RuntimeName);
+		return {};
+	}
+	
+	// Check if we have a DDC cache hit
+	FSharedBuffer RemoteData = GetFromDDC(FileId, RuntimeName, ModelDataIdentifier);
+	if (RemoteData.GetSize() > 0)
+	{
+		ModelData.Add(RuntimeName, TArray<uint8>((uint8*)RemoteData.GetData(), RemoteData.GetSize()));
+		
+		TArray<uint8>* CachedRemoteData = ModelData.Find(RuntimeName);
+		return TConstArrayView<uint8>(CachedRemoteData->GetData(), CachedRemoteData->Num());
+	}
+#endif //WITH_EDITOR
+
+	// Try to create the model
+	TArray<uint8> CreatedData = CreateRuntimeDataBlob(RuntimeName, FileType, FileData, FileId, nullptr);
+	if (CreatedData.Num() < 1)
+	{
+		return {};
+	}
+
+	// Cache the model
+	ModelData.Add(RuntimeName, CreatedData);
+
+#if WITH_EDITOR
+	// And put it into DDC
+	FSharedBuffer SharedBuffer = MakeSharedBufferFromArray(MoveTemp(CreatedData));
+	PutIntoDDC(FileId, RuntimeName, ModelDataIdentifier, SharedBuffer);
+#endif //WITH_EDITOR
+	
+	TArray<uint8>* CachedCreatedData = ModelData.Find(RuntimeName);
+	return TConstArrayView<uint8>(CachedCreatedData->GetData(), CachedCreatedData->Num());
+}
+
+void UNNEModelData::Serialize(FArchive& Ar)
+{
+	// Store the asset version (no effect in load)
+	Ar.UsingCustomVersion(UNNEModelData::GUID);
+
+#if WITH_EDITORONLY_DATA
+	// Recreate each model data when cooking
+	if (Ar.IsCooking() && Ar.IsSaving())
+	{
+		ModelData.Reset();
+
+		TArray<FString, TInlineAllocator<10>> CookedRuntimeNames;
+		CookedRuntimeNames.Append(GetTargetRuntimes());
+
+		//No target runtime means all currently registered ones.
+		if (GetTargetRuntimes().IsEmpty())
+		{
+			for (const TWeakInterfacePtr<INNERuntime>& Runtime : UE::NNE::GetAllRuntimes())
+			{
+				CookedRuntimeNames.Add(Runtime->GetRuntimeName());
+			}
+		}
+
+		for (const FString& RuntimeName : CookedRuntimeNames)
+		{
+			TArray<uint8> CreatedData = CreateRuntimeDataBlob(RuntimeName, FileType, FileData, FileId, Ar.GetArchiveState().CookingTarget());
+			if (CreatedData.Num() > 0)
+			{
+				ModelData.Add(RuntimeName, CreatedData);
+#if WITH_EDITOR
+				TWeakInterfacePtr<INNERuntime> NNERuntime = UE::NNE::GetRuntime<INNERuntime>(RuntimeName);
+				if (NNERuntime.IsValid())
+				{
+					FString ModelDataIdentifier = NNERuntime->GetModelDataIdentifier(FileType, FileData, FileId, Ar.GetArchiveState().CookingTarget());
+					if (ModelDataIdentifier.Len() > 0)
+					{
+						FSharedBuffer SharedBuffer = MakeSharedBufferFromArray(MoveTemp(CreatedData));
+						PutIntoDDC(FileId, RuntimeName, ModelDataIdentifier, SharedBuffer);
+					}
+					else
+					{
+						UE_LOG(LogNNE, Warning, TEXT("UNNEModelData: Runtime '%s' returned an empty string as a ModelDataIdentifier while cooking. GetModelDataIdentifier should always return a valid identifier."), *RuntimeName);
+					}
+				}
+				else
+				{
+					UE_LOG(LogNNE, Warning, TEXT("UNNEModelData: Runtime '%s' is among the cooked runtimes but instance is invalid."), *RuntimeName);
+				}
+#endif //WITH_EDITOR
+			}
+		}
+
+		// Dummy data for fields not required in the game
+		TArray<uint8> EmptyData;
+		TArray<FString> RuntimeNames;
+		ModelData.GetKeys(RuntimeNames);
+		int32 NumItems = RuntimeNames.Num();
+
+		Ar << FileType;
+		Ar << EmptyData;
+		Ar << FileId;
+		Ar << NumItems;
+
+		for (int i = 0; i < NumItems; i++)
+		{
+			Ar << RuntimeNames[i];
+			Ar << ModelData[RuntimeNames[i]];
+		}
+	}
+	else
+#endif //WITH_EDITORONLY_DATA
+	{
+		int32 NumItems = 0;
+
+#if WITH_EDITORONLY_DATA
+		if (Ar.CustomVer(UNNEModelData::GUID) >= TargetRuntimesAndAssetImportData)
+		{
+			Ar << TargetRuntimes;
+			Ar << AssetImportData;
+		}
+		else 
+		{
+			// AssetImportData should always be valid
+			AssetImportData = NewObject<UAssetImportData>(this, TEXT("AssetImportData"));
+		}
+#endif //WITH_EDITORONLY_DATA
+
+		Ar << FileType;
+		Ar << FileData;
+		Ar << FileId;
+		Ar << NumItems;
+
+		if (Ar.IsLoading())
+		{
+			for (int i = 0; i < NumItems; i++)
+			{
+				FString Name;
+				Ar << Name;
+				TArray<uint8> Data;
+				Ar << Data;
+				ModelData.Add(Name, MoveTemp(Data));
+			}
+		}
+	}
+}
+
+#if WITH_EDITORONLY_DATA
+namespace UE::NNE::ModelDataHelpers
+{
+	FString GetRuntimesAsString(TArrayView<const FString> Runtimes)
+	{
+		if (Runtimes.Num() == 0)
+		{
+			return TEXT("All");
+		}
+
+		FString RuntimesAsOneString;
+		bool bIsFirstRuntime = true;
+
+		for (const FString& Runtime : Runtimes)
+		{
+			if (!bIsFirstRuntime)
+			{
+				RuntimesAsOneString += TEXT(", ");
+			}
+			RuntimesAsOneString += Runtime;
+			bIsFirstRuntime = false;
+		}
+		return RuntimesAsOneString;
+	}
+} // UE::NNE::ModelDataHelpers
+
+void UNNEModelData::PostInitProperties()
+{
+	if (!HasAnyFlags(RF_ClassDefaultObject))
+	{
+		AssetImportData = NewObject<UAssetImportData>(this, TEXT("AssetImportData"));
+	}
+	Super::PostInitProperties();
+}
+
+void UNNEModelData::GetAssetRegistryTags(TArray<FAssetRegistryTag>& OutTags) const
+{
+	if (AssetImportData)
+	{
+		OutTags.Add(FAssetRegistryTag(SourceFileTagName(), AssetImportData->GetSourceData().ToJson(), FAssetRegistryTag::TT_Hidden));
+	}
+
+	OutTags.Add(FAssetRegistryTag("TargetRuntimes", UE::NNE::ModelDataHelpers::GetRuntimesAsString(GetTargetRuntimes()), FAssetRegistryTag::TT_Alphabetical));
+
+	Super::GetAssetRegistryTags(OutTags);
+}
+
+void UNNEModelData::SetTargetRuntimes(TArrayView<const FString> RuntimeNames)
+{
+	TargetRuntimes = RuntimeNames;
+
+	TArray<FString, TInlineAllocator<10>> CookedRuntimes;
+	ModelData.GetKeys(CookedRuntimes);
+	for (const FString& Runtime : CookedRuntimes)
+	{
+		if (!TargetRuntimes.Contains(Runtime))
+		{
+			ModelData.Remove(Runtime);
+		}
+	}
+	ModelData.Compact();
+}
+
+#endif //WITH_EDITORONLY_DATA
\ No newline at end of file
diff --git a/Engine/Source/Runtime/NNE/Private/NNETypes.cpp b/Engine/Source/Runtime/NNE/Private/NNETypes.cpp
new file mode 100644
--- /dev/null
+++ b/Engine/Source/Runtime/NNE/Private/NNETypes.cpp
@@ -0,0 +1,176 @@
+// Copyright Epic Games, Inc. All Rights Reserved.
+
+#include "NNETypes.h"
+#include "NNE.h"
+
+namespace UE::NNE::TypesUtils
+{
+	template <class T> bool AreShapeEqual(const T& lhs, const T& rhs)
+	{
+		check(lhs.Rank() <= lhs.MaxRank);
+
+		if (lhs.Rank() != rhs.Rank())
+			return false;
+
+		for (int32 i = 0; i < lhs.Rank(); ++i)
+		{
+			if (lhs.GetData()[i] != rhs.GetData()[i])
+				return false;
+		}
+			
+		return true;
+	}
+}
+
+namespace UE::NNE
+{
+
+bool FSymbolicTensorShape::IsConcrete() const
+{
+	for (int32 i = 0; i < Rank(); ++i)
+	{
+		if (Data[i] < 0)
+		{
+			return false;
+		}
+	}
+	return true;
+}
+
+bool FSymbolicTensorShape::operator==(const FSymbolicTensorShape& OtherShape) const
+{
+	return TypesUtils::AreShapeEqual(*this, OtherShape);
+}
+
+bool FSymbolicTensorShape::operator!=(const FSymbolicTensorShape& OtherShape) const { return !(*this == OtherShape); }
+
+void FSymbolicTensorShape::operator=(const FSymbolicTensorShape& OtherShape)
+{
+	check(OtherShape.Rank() <= MaxRank);
+	Data = OtherShape.Data;
+}
+
+FSymbolicTensorShape FSymbolicTensorShape::Make(TConstArrayView<int32> Data)
+{
+	if (Data.Num() > MaxRank)
+	{
+		UE_LOG(LogNNE, Warning, TEXT("Cannot create symbolic tensor shape, input is rank %d while max rank is %d"), Data.Num(), MaxRank);
+		return {};
+	}
+
+	FSymbolicTensorShape Shape;
+	Shape.Data.Append(Data);
+	return Shape;
+}
+
+uint64 FTensorShape::Volume() const
+{
+	uint64 Result = 1;
+
+	for (int32 Idx = 0; Idx < Data.Num(); ++Idx)
+	{
+		Result *= Data[Idx];
+	}
+
+	return Result;
+};
+
+bool FTensorShape::IsCompatibleWith(const FSymbolicTensorShape& SymbolicShape) const
+{
+	if (Rank() != SymbolicShape.Rank())
+	{
+		return false;
+	}
+	for (int32 i = 0; i < Rank(); ++i)
+	{
+		if (SymbolicShape.GetData()[i] >= 0 && SymbolicShape.GetData()[i] != Data[i])
+		{
+			return false;
+		}
+	}
+	return true;
+}
+
+bool FTensorShape::operator==(const FTensorShape& OtherShape) const
+{
+	return TypesUtils::AreShapeEqual(*this, OtherShape);
+}
+
+bool FTensorShape::operator!=(const FTensorShape& OtherShape) const { return !(*this == OtherShape); }
+
+void FTensorShape::operator=(const FTensorShape& OtherShape)
+{
+	check(OtherShape.Rank() <= MaxRank);
+	Data = OtherShape.Data;
+}
+
+FTensorShape FTensorShape::Make(TConstArrayView<uint32> Data)
+{
+	if (Data.Num() > MaxRank)
+	{
+		UE_LOG(LogNNE, Warning, TEXT("Cannot create tensor shape, input is rank %d while max rank is %d"), Data.Num(), MaxRank);
+		return {};
+	}
+
+	FTensorShape Shape;
+	Shape.Data.Append(Data);
+	return Shape;
+}
+	
+FTensorShape FTensorShape::MakeFromSymbolic(const FSymbolicTensorShape& SymbolicShape)
+{
+	FTensorShape ConcreteShape;
+	for (int32 i = 0; i < SymbolicShape.GetData().Num(); ++i)
+	{
+		int32 Dim = SymbolicShape.GetData()[i];
+		ConcreteShape.Data.Add(Dim < 0 ? 1 : Dim);
+	}
+	return ConcreteShape;
+}
+
+size_t GetTensorDataTypeSizeInBytes(ENNETensorDataType InType)
+{
+	switch (InType)
+	{
+	case ENNETensorDataType::Complex128:
+		return 16;
+
+	case ENNETensorDataType::Complex64:
+		return 8;
+
+	case ENNETensorDataType::Double:
+	case ENNETensorDataType::Int64:
+	case ENNETensorDataType::UInt64:
+		return 8;
+
+	case ENNETensorDataType::Float:
+	case ENNETensorDataType::Int32:
+	case ENNETensorDataType::UInt32:
+		return 4;
+
+	case ENNETensorDataType::Half:
+	case ENNETensorDataType::BFloat16:
+	case ENNETensorDataType::Int16:
+	case ENNETensorDataType::UInt16:
+		return 2;
+
+	case ENNETensorDataType::Int8:
+	case ENNETensorDataType::UInt8:
+	case ENNETensorDataType::Char:
+	case ENNETensorDataType::Boolean:
+		return 1;
+	}
+
+	return 0;
+}
+
+FTensorDesc FTensorDesc::Make(const FString& Name, const FSymbolicTensorShape& Shape, ENNETensorDataType DataType)
+{
+	FTensorDesc Desc;
+	Desc.Name = Name;
+	Desc.DataType = DataType;
+	Desc.Shape = Shape;
+	return Desc;
+}
+
+} // namespace UE::NNE
diff --git a/Engine/Source/Runtime/NNE/Public/NNE.h b/Engine/Source/Runtime/NNE/Public/NNE.h
new file mode 100644
--- /dev/null
+++ b/Engine/Source/Runtime/NNE/Public/NNE.h
@@ -0,0 +1,60 @@
+// Copyright Epic Games, Inc. All Rights Reserved.
+
+#pragma once
+
+#include "CoreMinimal.h"
+#include "NNERuntime.h"
+#include "UObject/WeakInterfacePtr.h"
+
+NNE_API DECLARE_LOG_CATEGORY_EXTERN(LogNNE, Log, All);
+
+namespace UE::NNE
+{
+	/**
+	 * Register a runtime to make it accessible to NNE clients.
+	 *
+	 * The caller needs to keep a strong pointer to the runtime object to prevent it from being garbage collected.
+	 *
+	 * @param Runtime A weak interface pointer to the runtime to be registered.
+	 * @return True if the runtime has been registered successfully, false otherwise.
+	 */
+	NNE_API bool RegisterRuntime(TWeakInterfacePtr<INNERuntime> Runtime);
+
+	/**
+	 * Unregister a registered runtime.
+	 *
+	 * @param Runtime A weak interface pointer to the runtime to be unregistered.
+	 * @return True if the runtime has been unregistered successfully, false otherwise (e.g. if the runtime has not been registered).
+	 */
+	NNE_API bool UnregisterRuntime(TWeakInterfacePtr<INNERuntime> Runtime);
+	
+	/**
+	 * List and return all registered runtimes.
+	 *
+	 * @return An array containing weak pointers to all registered runtimes.
+	 */
+	NNE_API TArrayView<TWeakInterfacePtr<INNERuntime>> GetAllRuntimes();
+
+	/**
+	 * Find and return a runtime by name and interface.
+	 *
+	 * This function tries to find a runtime by name and casts it to the interface passed as template argument (e.g. INNERuntime, INNERuntimeCPU or INNERuntimeRDG).
+	 *
+	 * @param Name The name of the runtime.
+	 * @return A weak pointer to the runtime if it has been found and implements the interface in the template argument or an invalid pointer otherwise.
+	 */
+	template<class T> TWeakInterfacePtr<T> GetRuntime(const FString& Name)
+	{
+		for (TWeakInterfacePtr<INNERuntime> Runtime : GetAllRuntimes())
+		{
+			if (Runtime->GetRuntimeName() == Name)
+			{
+				T* RuntimePtr = Cast<T>(Runtime.Get());
+				return TWeakInterfacePtr<T>(RuntimePtr);
+			}
+				
+		}
+		return TWeakInterfacePtr<T>(nullptr);
+	}
+	
+}
diff --git a/Engine/Source/Runtime/NNE/Public/NNEModelData.h b/Engine/Source/Runtime/NNE/Public/NNEModelData.h
new file mode 100644
--- /dev/null
+++ b/Engine/Source/Runtime/NNE/Public/NNEModelData.h
@@ -0,0 +1,136 @@
+// Copyright Epic Games, Inc. All Rights Reserved.
+
+#pragma once
+
+#include "CoreMinimal.h"
+
+#include "NNEModelData.generated.h"
+
+/**
+ * This class represents assets that store neural network model data.
+ *
+ * Neural network models typically consist of a graph of operations and corresponding parameters as e.g. weights.
+ * UNNEModelData assets store such model data as imported e.g. by the UNNEModelDataFactory class.
+ * An INNERuntime object retrieved by UE::NNE::GetRuntime<T>(const FString& Name) can be used to create an inferable neural network model.
+ */
+UCLASS(BlueprintType, Category = "NNE")
+class NNE_API UNNEModelData : public UObject
+{
+	GENERATED_BODY()
+
+public:
+
+	/**
+	 * Initialize the model data with a copy of the data inside Buffer.
+	 *
+	 * This function is called by the UNNEModelDataFactory class when importing a neural network model file.
+	 *
+	 * @param Type A string identifying the type of data inside this asset. Corresponds to the extension of the imported file.
+	 * @param Buffer The raw binary file data of the imported model to be copied into this asset.
+	 */
+	void Init(const FString& Type, TConstArrayView<uint8> Buffer);
+
+	/**
+	 * Get the type of data inside FileData.
+	 *
+	 * The FileType identifies the type of data inside FileData and typically is the extension of the file used to create the asset.
+	 *
+	 * @return The FileType.
+	 */
+	FString GetFileType();
+
+	/**
+	 * Get read only access to FileData.
+	 *
+	 * The FileData contains the binary data of the file which has been used to create the asset.
+	 *
+	 * @return The FileData.
+	 */
+	TConstArrayView<uint8> GetFileData();
+
+	/**
+	 * Get the FGuid identifying the FileData.
+	 *
+	 * The FileId is created on import of an asset. It can be used to identify the FileData, e.g. when putting corresponding data into the DDC or caching data locally.
+	 *
+	 * @return The FileId.
+	 */
+	FGuid GetFileId();
+
+	/**
+	 * Get the cached (editor) or cooked (game) optimized model data for a given runtime.
+	 *
+	 * This function is used by runtimes when creating a model. In editor, the function will create the optimized model data with the passed runtime in case it has not been cached in the DCC yet. In game, the cooked data is accessed.
+	 *
+	 * @param RuntimeName The name of the runtime for which the data should be returned.
+	 * @return The optimized and runtime specific model data or an empty view in case of a failure.
+	 */
+	TConstArrayView<uint8> GetModelData(const FString& RuntimeName);
+
+	/**
+	 * Implements custom serialization of this asset.
+	 * @param Ar The archive to serialize from/to.
+	 */
+	virtual void Serialize(FArchive& Ar) override;
+
+	/**
+	 * A Guid used for asset versioning.
+	 */
+	const static FGuid GUID;
+
+#if WITH_EDITORONLY_DATA
+	// UObject interface
+	virtual void PostInitProperties() override;
+	virtual void GetAssetRegistryTags(TArray<FAssetRegistryTag>& OutTags) const override;
+	// End of UObject interface
+
+	/**
+	 * Get the target runtimes this model data will be cooked for. An empty list mean all runtimes.
+	 *
+	 * @return The target runtimes names.
+	 */
+	TArrayView<const FString> GetTargetRuntimes() const { return TargetRuntimes; }
+
+	/**
+	 * Set the target runtimes this model data will be cooked for. An empty list mean all runtimes.
+	 *
+	 * @param RuntimeNames The target runtimes names.
+	 */
+	void SetTargetRuntimes(TArrayView<const FString> RuntimeNames);
+
+	/**
+	 * Importing data used for this asset.
+	 */
+	UPROPERTY(VisibleAnywhere, Instanced, Category = ImportSettings)
+	TObjectPtr<class UAssetImportData> AssetImportData;
+
+private:
+	/**
+	 * A list of string of the supported runtime, empty to support them all.
+	 */
+	TArray<FString> TargetRuntimes;
+#endif // WITH_EDITORONLY_DATA
+
+private:
+
+	/**
+	 * A string identifying the type of data inside this asset. Corresponds to the extension of the imported file.
+	 */
+	FString FileType;
+
+	/**
+	 * The raw binary file data of the imported model.
+	 */
+	TArray<uint8> FileData;
+
+	/**
+	 * A Guid that uniquely identifies this model. This is used to cache optimized models in the editor.
+	 */
+	FGuid FileId;
+
+	/**
+	 * The processed / optimized model data for the different runtimes.
+	 */
+	TMap<FString, TArray<uint8>> ModelData;
+
+};
diff --git a/Engine/Source/Runtime/NNE/Public/NNERuntime.h b/Engine/Source/Runtime/NNE/Public/NNERuntime.h
new file mode 100644
--- /dev/null
+++ b/Engine/Source/Runtime/NNE/Public/NNERuntime.h
@@ -0,0 +1,80 @@
+// Copyright Epic Games, Inc. All Rights Reserved.
+
+#pragma once
+
+#include "CoreMinimal.h"
+#include "UObject/Interface.h"
+
+#include "NNERuntime.generated.h"
+
+class ITargetPlatform;
+
+UINTERFACE()
+class NNE_API UNNERuntime : public UInterface
+{
+	GENERATED_BODY()
+};
+
+/**
+ * The base interface of a neural network runtime. It is returned by UE::NNE::GetAllRuntimes() giving access to all available runtimes.
+ *
+ * This interface is mainly used internally by UNNEModelData to cook a model given in a file format (e.g. .onnx) into runtime specific model data.
+ * The model data is then stored inside UNNEModelData which then can be used in game to create an inferable model.
+ * See INNERuntimeCPU and INNERuntimeRDG on how to create a model given a UNNEModelData asset.
+ */
+class NNE_API INNERuntime
+{
+	GENERATED_BODY()
+
+public:
+
+	/**
+	 * Get the name of the runtime.
+	 *
+	 * @return The name of the runtime.
+	 */
+	virtual FString GetRuntimeName() const = 0;
+	
+	/**
+	 * Check if the runtime is supported on a specific platform.
+	 *
+	 * This function is used internally in the cooking process of UNNEModelData.
+	 *
+	 * @param TargetPlatform Interface identifying the target platforms.
+	 * @return True if the platform is supported, false otherwise.
+	 */
+	virtual bool IsPlatformSupported(const ITargetPlatform* TargetPlatform) const = 0;
+
+	/**
+	 * Check if the runtime is able to create model data given some file data representing a neural network.
+	 *
+	 * @param FileType The type of file inside FileData. Corresponds to the file extension (e.g. 'onnx').
+	 * @param FileData The raw binary file of a neural network model.
+	 * @param FileId The unique identifier representing FileData.
+	 * @param TargetPlatform The Interface identifying the target platform for which the data needs to be created. A null pointer indicates the currently compiled/running platform.
+	 * @return True if the runtime is able to create model data, false otherwise.
+	 */
+	virtual bool CanCreateModelData(FString FileType, TConstArrayView<uint8> FileData, FGuid FileId, const ITargetPlatform* TargetPlatform) const = 0;
+
+	/**
+	 * Create model data given some raw file data.
+	 *
+	 * @param FileType The type of file inside FileData. Corresponds to the file extension (e.g. 'onnx').
+	 * @param FileData The raw binary file of a neural network model.
+	 * @param FileId The unique identifier representing FileData.
+	 * @param TargetPlatform The Interface identifying the target platform for which the data needs to be created. A null pointer indicates the currently compiled/running platform.
+	 * @return Data representing the runtime specific representation of the model to be stored by UNNEModelData on success or an empty array otherwise.
+	 */
+	virtual TArray<uint8> CreateModelData(FString FileType, TConstArrayView<uint8> FileData, FGuid FileId, const ITargetPlatform* TargetPlatform) = 0;
+
+	/**
+	 * Get an id uniquely identifying the model data. This is used by UNNEModelData to create a FCacheKey to get and add the model data from and to a DDC.
+	 *
+	 * @param FileType The type of file inside FileData. Corresponds to the file extension (e.g. 'onnx').
+	 * @param FileData The raw binary file of a neural network model.
+	 * @param FileId The unique identifier representing FileData.
+	 * @param TargetPlatform The Interface identifying the target platform for which the data is produced or by which it will be consumed. A null pointer indicates the currently compiled/running platform.
+	 * @return A string that identifies the model data uniquely, e.g. by containing the FileId, model data version and target platform display name.
+	 */
+	virtual FString GetModelDataIdentifier(FString FileType, TConstArrayView<uint8> FileData, FGuid FileId, const ITargetPlatform* TargetPlatform) = 0;
+};
\ No newline at end of file
diff --git a/Engine/Source/Runtime/NNE/Public/NNERuntimeCPU.h b/Engine/Source/Runtime/NNE/Public/NNERuntimeCPU.h
new file mode 100644
--- /dev/null
+++ b/Engine/Source/Runtime/NNE/Public/NNERuntimeCPU.h
@@ -0,0 +1,160 @@
+// Copyright Epic Games, Inc. All Rights Reserved.
+
+#pragma once
+
+#include "CoreMinimal.h"
+#include "NNERuntime.h"
+#include "NNETypes.h"
+#include "UObject/Interface.h"
+
+#include "NNERuntimeCPU.generated.h"
+
+class UNNEModelData;
+
+namespace UE::NNE
+{
+
+/**
+ * The tensor binding for passing input and output to IModelCPU.
+ *
+ * Memory is owned by the caller. The caller must make sure the buffer is large enough and at least as large as SizeInBytes.
+ */
+struct NNE_API FTensorBindingCPU
+{
+	void*	Data;
+	uint64	SizeInBytes;
+};
+
+/**
+ * The interface of a model instance that can run on CPU.
+ *
+ * Use UE::NNE::IModelCPU::CreateModelInstance() to get a model instance.
+ * Use UE::NNE::GetRuntime<INNERuntimeCPU>(RuntimeName) to get a runtime capable of creating CPU models.
+ */
+class NNE_API IModelInstanceCPU
+{
+public:
+
+	virtual ~IModelInstanceCPU() = default;
+
+	/**
+	 * Get the input tensor descriptions as defined by the model, potentially with variable dimensions.
+	 *
+	 * @return An array containing a tensor descriptor for each input tensor of the model.
+	 */
+	virtual TConstArrayView<FTensorDesc> GetInputTensorDescs() const = 0;
+	
+	/**
+	 * Get the output tensor descriptions as defined by the model, potentially with variable dimensions.
+	 *
+	 * @return An array containing a tensor descriptor for each output tensor of the model.
+	 */
+	virtual TConstArrayView<FTensorDesc> GetOutputTensorDescs() const = 0;
+
+	/**
+	 * Get the input shapes.
+	 *
+	 * SetInputTensorShapes must be called prior of running a model.
+	 *
+	 * @return An array of input shapes or an empty array if SetInputTensorShapes has not been called.
+	 */
+	virtual TConstArrayView<FTensorShape> GetInputTensorShapes() const = 0;
+
+	/**
+	 * Getters for outputs shapes if they were already resolved.
+	 *
+	 * Output shapes might be resolved after a call to SetInputTensorShapes if the model and runtime supports it.
+	 * Otherwise they will be resolved while running the model
+	 *
+	 * @return An array of output shapes or an empty array if not resolved yet.
+	 */
+	virtual TConstArrayView<FTensorShape> GetOutputTensorShapes() const = 0;
+	
+	/**
+	 * Prepare the model to be run with the given input shape.
+	 *
+	 * The call is mandatory before a model can be run.
+	 * The function will run shape inference and resolve, if possible, the output shapes which can then be accessed by calling GetOutputTensorShapes().
+	 * This is a potentially expensive call and should be called lazily if possible.
+	 *
+	 * @param InInputShapes The input shapes to prepare the model with.
+	 * @return 0 on success or a non-zero number otherwise.
+	 */
+	virtual int SetInputTensorShapes(TConstArrayView<FTensorShape> InInputShapes) = 0;
+
+	/**
+	 * Evaluate the model synchronously.
+	 *
+	 * SetInputTensorShapes must be called prior to this call.
+	 * This function will block the calling thread until the inference is complete.
+	 * The caller owns the memory inside the bindings and must make sure that they are big enough.
+	 * Clients can call this function from an async task but must make sure the memory remains valid throughout the evaluation.
+	 *
+	 * @param InInputTensors An array containing tensor bindings for each input tensor with caller owned memory containing the input data.
+	 * @param InOutputTensors An array containing tensor bindings for each output tensor with caller owned memory big enough to contain the results on success.
+	 * @return 0 on success or a non-zero number otherwise.
+	 */
+	virtual int RunSync(TConstArrayView<FTensorBindingCPU> InInputTensors, TConstArrayView<FTensorBindingCPU> InOutputTensors) = 0;
+};
+
+/**
+ * The interface of a model capable of creating model instance that can run on CPU.
+ *
+ * Use UE::NNE::GetRuntime<INNERuntimeCPU>(RuntimeName) to get a runtime capable of creating CPU models.
+ */
+class NNE_API IModelCPU
+{
+public:
+
+	virtual ~IModelCPU() = default;
+
+	/**
+	 * Create a model instance for inference
+	 *
+	 * The runtime have the opportunity to share the model weights among multiple IModelInstanceCPU created from an IModelCPU instance, however this is not mandatory.
+	 * The caller can decide to convert the result into a shared pointer if required (e.g. if the model needs to be shared with an async task for evaluation).
+	 *
+	 * @return A caller owned model representing the neural network instance created.
+	 */
+	virtual TSharedPtr<UE::NNE::IModelInstanceCPU> CreateModelInstanceCPU() = 0;
+};
+
+} // UE::NNE
+
+UINTERFACE()
+class NNE_API UNNERuntimeCPU : public UInterface
+{
+	GENERATED_BODY()
+};
+
+/**
+ * The interface of a neural network runtime capable of creating CPU models.
+ *
+ * Call UE::NNE::GetRuntime<INNERuntimeCPU>(RuntimeName) to get a runtime implementing this interface.
+ */
+class NNE_API INNERuntimeCPU
+{
+	GENERATED_BODY()
+	
+public:
+
+	/**
+	 * Check if the runtime is able to create a model given some ModelData.
+	 *
+	 * @param ModelData The model data for which to create a model.
+	 * @return True if the runtime is able to create the model, false otherwise.
+	 */
+	virtual bool CanCreateModelCPU(TObjectPtr<UNNEModelData> ModelData) const = 0;
+	
+	/**
+	 * Create a model given some ModelData.
+	 *
+	 * The caller must make sure ModelData remains valid throughout the call.
+	 * ModelData is not required anymore after the model has been created.
+	 * The caller can decide to convert the result into a shared pointer if required (e.g. if the model needs to be shared with an async task for evaluation).
+	 *
+	 * @param ModelData The model data for which to create a model.
+	 * @return A caller owned model representing the neural network created from ModelData.
+	 */
+	virtual TSharedPtr<UE::NNE::IModelCPU> CreateModelCPU(TObjectPtr<UNNEModelData> ModelData) = 0;
+};
\ No newline at end of file
diff --git a/Engine/Source/Runtime/NNE/Public/NNERuntimeGPU.h b/Engine/Source/Runtime/NNE/Public/NNERuntimeGPU.h
new file mode 100644
--- /dev/null
+++ b/Engine/Source/Runtime/NNE/Public/NNERuntimeGPU.h
@@ -0,0 +1,160 @@
+// Copyright Epic Games, Inc. All Rights Reserved.
+
+#pragma once
+
+#include "CoreMinimal.h"
+#include "NNERuntime.h"
+#include "NNETypes.h"
+#include "UObject/Interface.h"
+
+#include "NNERuntimeGPU.generated.h"
+
+class UNNEModelData;
+
+namespace UE::NNE
+{
+
+/**
+ * The tensor binding for passing input and output to IModelGPU.
+ *
+ * Memory is owned by the caller. The caller must make sure the buffer is large enough and at least as large as SizeInBytes.
+ */
+struct NNE_API FTensorBindingGPU
+{
+	void*	Data;
+	uint64	SizeInBytes;
+};
+
+/**
+ * The interface of a model instance that can run on GPU.
+ *
+ * Use UE::NNE::IModelGPU::CreateModelInstance() to get a model instance.
+ * Use UE::NNE::GetRuntime<INNERuntimeGPU>(RuntimeName) to get a runtime capable of creating GPU models.
+ */
+class NNE_API IModelInstanceGPU
+{
+public:
+
+	virtual ~IModelInstanceGPU() = default;
+
+	/**
+	 * Get the input tensor descriptions as defined by the model, potentially with variable dimensions.
+	 *
+	 * @return An array containing a tensor descriptor for each input tensor of the model.
+	 */
+	virtual TConstArrayView<FTensorDesc> GetInputTensorDescs() const = 0;
+	
+	/**
+	 * Get the output tensor descriptions as defined by the model, potentially with variable dimensions.
+	 *
+	 * @return An array containing a tensor descriptor for each output tensor of the model.
+	 */
+	virtual TConstArrayView<FTensorDesc> GetOutputTensorDescs() const = 0;
+
+	/**
+	 * Get the input shapes.
+	 *
+	 * SetInputTensorShapes must be called prior of running a model.
+	 *
+	 * @return An array of input shapes or an empty array if SetInputTensorShapes has not been called.
+	 */
+	virtual TConstArrayView<FTensorShape> GetInputTensorShapes() const = 0;
+
+	/**
+	 * Getters for outputs shapes if they were already resolved.
+	 *
+	 * Output shapes might be resolved after a call to SetInputTensorShapes if the model and runtime supports it.
+	 * Otherwise they will be resolved while running the model
+	 *
+	 * @return An array of output shapes or an empty array if not resolved yet.
+	 */
+	virtual TConstArrayView<FTensorShape> GetOutputTensorShapes() const = 0;
+	
+	/**
+	 * Prepare the model to be run with the given input shape.
+	 *
+	 * The call is mandatory before a model can be run.
+	 * The function will run shape inference and resolve, if possible, the output shapes which can then be accessed by calling GetOutputTensorShapes().
+	 * This is a potentially expensive call and should be called lazily if possible.
+	 *
+	 * @param InInputShapes The input shapes to prepare the model with.
+	 * @return 0 on success or a non-zero number otherwise.
+	 */
+	virtual int SetInputTensorShapes(TConstArrayView<FTensorShape> InInputShapes) = 0;
+
+	/**
+	 * Evaluate the model synchronously.
+	 *
+	 * SetInputTensorShapes must be called prior to this call.
+	 * This function will block the calling thread until the inference is complete.
+	 * The caller owns the memory inside the bindings and must make sure that they are big enough.
+	 * Clients can call this function from an async task but must make sure the memory remains valid throughout the evaluation.
+	 *
+	 * @param InInputTensors An array containing tensor bindings for each input tensor with caller owned memory containing the input data.
+	 * @param InOutputTensors An array containing tensor bindings for each output tensor with caller owned memory big enough to contain the results on success.
+	 * @return 0 on success or a non-zero number otherwise.
+	 */
+	virtual int RunSync(TConstArrayView<FTensorBindingGPU> InInputTensors, TConstArrayView<FTensorBindingGPU> InOutputTensors) = 0;
+};
+
+/**
+ * The interface of a model capable of creating model instance that can run on GPU.
+ *
+ * Use UE::NNE::GetRuntime<INNERuntimeGPU>(RuntimeName) to get a runtime capable of creating GPU models.
+ */
+class NNE_API IModelGPU
+{
+public:
+
+	virtual ~IModelGPU() = default;
+
+	/**
+	 * Create a model instance for inference
+	 *
+	 * The runtime have the opportunity to share the model weights among multiple IModelInstanceGPU created from an IModelGPU instance, however this is not mandatory.
+	 * The caller can decide to convert the result into a shared pointer if required (e.g. if the model needs to be shared with an async task for evaluation).
+	 *
+	 * @return A caller owned model representing the neural network instance created.
+	 */
+	virtual TSharedPtr<UE::NNE::IModelInstanceGPU> CreateModelInstanceGPU() = 0;
+};
+
+} // UE::NNE
+
+UINTERFACE()
+class NNE_API UNNERuntimeGPU : public UInterface
+{
+	GENERATED_BODY()
+};
+
+/**
+ * The interface of a neural network runtime capable of creating GPU models.
+ *
+ * Call UE::NNE::GetRuntime<INNERuntimeGPU>(RuntimeName) to get a runtime implementing this interface.
+ */
+class NNE_API INNERuntimeGPU
+{
+	GENERATED_BODY()
+	
+public:
+
+	/**
+	 * Check if the runtime is able to create a model given some ModelData.
+	 *
+	 * @param ModelData The model data for which to create a model.
+	 * @return True if the runtime is able to create the model, false otherwise.
+	 */
+	virtual bool CanCreateModelGPU(TObjectPtr<UNNEModelData> ModelData) const = 0;
+	
+	/**
+	 * Create a model given some ModelData.
+	 *
+	 * The caller must make sure ModelData remains valid throughout the call.
+	 * ModelData is not required anymore after the model has been created.
+	 * The caller can decide to convert the result into a shared pointer if required (e.g. if the model needs to be shared with an async task for evaluation).
+	 *
+	 * @param ModelData The model data for which to create a model.
+	 * @return A caller owned model representing the neural network created from ModelData.
+	 */
+	virtual TSharedPtr<UE::NNE::IModelGPU> CreateModelGPU(TObjectPtr<UNNEModelData> ModelData) = 0;
+};
\ No newline at end of file
diff --git a/Engine/Source/Runtime/NNE/Public/NNERuntimeRDG.h b/Engine/Source/Runtime/NNE/Public/NNERuntimeRDG.h
new file mode 100644
--- /dev/null
+++ b/Engine/Source/Runtime/NNE/Public/NNERuntimeRDG.h
@@ -0,0 +1,160 @@
+// Copyright Epic Games, Inc. All Rights Reserved.
+
+#pragma once
+
+#include "CoreMinimal.h"
+#include "NNETypes.h"
+#include "RenderGraphFwd.h"
+#include "UObject/Interface.h"
+
+#include "NNERuntimeRDG.generated.h"
+
+class UNNEModelData;
+
+namespace UE::NNE
+{
+
+/**
+ * The tensor binding for passing input and output to IModelRDG.
+ *
+ * Memory is owned by the caller. The caller must make sure the buffer is large enough.
+ */
+struct NNE_API FTensorBindingRDG
+{
+	FRDGBufferRef Buffer;
+};
+
+/**
+ * The interface of a model instance that can run on RDG.
+ *
+ * Use UE::NNE::IModelRDG::CreateModelInstance() to get a model instance.
+ * Use UE::NNE::GetRuntime<INNERuntimeRDG>(RuntimeName) to get a runtime capable of creating RDG models.
+ */
+class NNE_API IModelInstanceRDG
+{
+public:
+
+	virtual ~IModelInstanceRDG() = default;
+
+	/**
+	 * Get the input tensor descriptions as defined by the model, potentially with variable dimensions.
+	 *
+	 * @return An array containing a tensor descriptor for each input tensor of the model.
+	 */
+	virtual TConstArrayView<FTensorDesc> GetInputTensorDescs() const = 0;
+
+	/**
+	 * Get the output tensor descriptions as defined by the model, potentially with variable dimensions.
+	 *
+	 * @return An array containing a tensor descriptor for each output tensor of the model.
+	 */
+	virtual TConstArrayView<FTensorDesc> GetOutputTensorDescs() const = 0;
+
+	/**
+	 * Get the input shapes.
+	 *
+	 * SetInputTensorShapes must be called prior of running a model.
+	 *
+	 * @return An array of input shapes or an empty array if SetInputTensorShapes has not been called.
+	 */
+	virtual TConstArrayView<FTensorShape> GetInputTensorShapes() const = 0;
+
+	/**
+	 * Getters for outputs shapes if they were already resolved.
+	 *
+	 * Output shapes might be resolved after a call to SetInputTensorShapes if the model and runtime supports it.
+	 * Otherwise they will be resolved while running the model
+	 *
+	 * @return An array of output shapes or an empty array if not resolved yet.
+	 */
+	virtual TConstArrayView<FTensorShape> GetOutputTensorShapes() const = 0;
+	
+	/**
+	 * Prepare the model to be run with the given input shape.
+	 * 
+	 * The call is mandatory before a model can be run.
+	 * This function might be called from the render thread, if not it is up to the caller to ensure thread safety.
+	 * The function will run shape inference and resolve, if possible, the output shapes which can then be accessed by calling GetOutputTensorShapes().
+	 * This is a potentially expensive call and should be called lazily if possible.
+	 *
+	 * @param InInputShapes The input shapes to prepare the model with.
+	 * @return 0 on success or a non-zero number otherwise.
+	 */
+	virtual int SetInputTensorShapes(TConstArrayView<FTensorShape> InInputShapes) = 0;
+
+	/**
+	 * Enqueue the model graph to a FRDGBuilder.
+	 *
+	 * This function must be called from the render thread.
+	 * SetInputTensorShapes must be called prior to this call.
+	 * The caller owns the memory inside the bindings and must make sure that they are big enough.
+	 *
+	 * @param RDGBuilder The RDG builder to which the neural network operations are enqueued.
+	 * @param InInputTensors An array containing tensor bindings for each input tensor with caller owned memory containing the input data.
+	 * @param InOutputTensors An array containing tensor bindings for each output tensor with caller owned memory big enough to contain the results on success.
+	 * @return 0 on success or a non-zero number otherwise.
+	 */
+	virtual int EnqueueRDG(FRDGBuilder& RDGBuilder, TConstArrayView<FTensorBindingRDG> Inputs, TConstArrayView<FTensorBindingRDG> Outputs) = 0;
+};
+
+/**
+ * The interface of a model capable of creating model instance that can run on RDG.
+ *
+ * Use UE::NNE::GetRuntime<INNERuntimeRDG>(RuntimeName) to get a runtime capable of creating RDG models.
+ */
+class NNE_API IModelRDG
+{
+public:
+
+	virtual ~IModelRDG() = default;
+
+	/**
+	 * Create a model instance for inference
+	 *
+	 * The runtime have the opportunity to share the model weights among multiple IModelInstanceRDG created from an IModelRDG instance, however this is not mandatory.
+	 * The caller can decide to convert the result into a shared pointer if required (e.g. if the model needs to be shared with an async task for evaluation).
+	 *
+	 * @return A caller owned model representing the neural network instance created.
+	 */
+	virtual TSharedPtr<UE::NNE::IModelInstanceRDG> CreateModelInstanceRDG() = 0;
+};
+
+} // UE::NNE
+
+UINTERFACE()
+class NNE_API UNNERuntimeRDG : public UInterface
+{
+	GENERATED_BODY()
+};
+
+/**
+ * The interface of a neural network runtime capable of creating RDG models.
+ *
+ * Call UE::NNE::GetRuntime<INNERuntimeRDG>(RuntimeName) to get a runtime implementing this interface.
+ */
+class NNE_API INNERuntimeRDG
+{
+	GENERATED_BODY()
+
+public:
+
+	/**
+	 * Check if the runtime is able to create a model given some ModelData.
+	 *
+	 * @param ModelData The model data for which to create a model.
+	 * @return True if the runtime is able to create the model, false otherwise.
+	 */
+	virtual bool CanCreateModelRDG(TObjectPtr<UNNEModelData> ModelData) const = 0;
+	
+	/**
+	 * Create a model given some ModelData.
+	 *
+	 * The caller must make sure ModelData remains valid throughout the call.
+	 * ModelData is not required anymore after the model has been created.
+	 * The caller can decide to convert the result into a shared pointer if required (e.g. if the model needs to be shared with the render thread).
+	 *
+	 * @param ModelData The model data for which to create a model.
+	 * @return A caller owned model representing the neural network created from ModelData.
+	 */
+	virtual TSharedPtr<UE::NNE::IModelRDG> CreateModelRDG(TObjectPtr<UNNEModelData> ModelData) = 0;
+};
\ No newline at end of file
diff --git a/Engine/Source/Runtime/NNE/Public/NNETypes.h b/Engine/Source/Runtime/NNE/Public/NNETypes.h
new file mode 100644
--- /dev/null
+++ b/Engine/Source/Runtime/NNE/Public/NNETypes.h
@@ -0,0 +1,280 @@
+// Copyright Epic Games, Inc. All Rights Reserved.
+
+#pragma once
+
+#include "CoreMinimal.h"
+#include "UObject/Class.h"
+
+#include "NNETypes.generated.h"
+
+/**
+ * The enum lists all tensor data types used in NNE 
+ */
+UENUM()
+enum class ENNETensorDataType : uint8
+{
+	None,
+	Char,								//!< Character type
+	Boolean,							//!< Boolean type
+	Half,								//!< 16-bit floating number
+	Float,								//!< 32-bit floating number
+	Double,								//!< 64-bit floating number
+	Int8,								//!< 8-bit signed integer
+	Int16,								//!< 16-bit signed integer
+	Int32,								//!< 32-bit signed integer
+	Int64,								//!< 64-bit signed integer
+	UInt8,								//!< 8-bit unsigned integer
+	UInt16,								//!< 16-bit unsigned integer
+	UInt32,								//!< 32-bit unsigned integer
+	UInt64,								//!< 64-bit unsigned integer
+	Complex64,							//!< 64-bit Complex Number
+	Complex128,							//!< 128-bit Complex Number
+	BFloat16							//!< 16-bit floating number
+};
+
+namespace UE::NNE
+{
+
+	/**
+	 * A symbolic tensor shape represents the shape of a tensor with potentially variable dimension.
+	 * 
+	 * The variable dimensions are represented by -1 values. 
+	 */
+	class NNE_API FSymbolicTensorShape
+	{
+	public:
+		/**
+		 * The maximum number of dimensions supported by the symbolic tensor shape, aka the rank.
+		 */
+		constexpr static int32	MaxRank = 8;
+
+	private:
+		TArray<int32, TInlineAllocator<MaxRank>> Data;
+
+	public:
+		/**
+		 * Default construct a symbolic tensor shape.
+		 */
+		FSymbolicTensorShape() = default;
+
+		/**
+		 * Construct this symbolic tensor shape from another one.
+		 * 
+		 * @param OtherShape the other symbolic tensor shape to copy from.
+		 */
+		FSymbolicTensorShape(const FSymbolicTensorShape& OtherShape) = default;
+
+		/**
+		 * Construct a symbolic tensor shape with the given dimensions.
+		 *
+		 * @param Data an array of dimensions.
+		 * @return a symbolic tensor shape with the given dimensions.
+		 */
+		static FSymbolicTensorShape Make(TConstArrayView<int32> Data);
+
+		/**
+		 * Get the dimensions of the symbolic tensor shape.
+		 *
+		 * @return a view on the dimensions of the symbolic tensor shape.
+		 */
+		inline TConstArrayView<int32> GetData() const { return Data; };
+		
+		/**
+		 * Get the number of dimensions of the symbolic tensor shape, aka the rank.
+		 *
+		 * @return the number of dimensions of the symbolic tensor shape.
+		 */
+		inline int32 Rank() const { return Data.Num(); }
+
+		/**
+		 * Check if the symbolic tensor shape has no variable dimensions, aka is concrete.
+		 *
+		 * @return true if the symbolic tensor shape is concrete, false otherwise.
+		 */
+		bool IsConcrete() const;
+
+		/**
+		 * Check if the symbolic tensor shape is equal to another one.
+		 *
+		 * @param OtherShape the other symbolic tensor shape.
+		 * @return true if the symbolic tensor shape is equal to the other one, false otherwise.
+		 */
+		bool operator==(const FSymbolicTensorShape& OtherShape) const;
+		
+		/**
+		 * Check if the symbolic tensor shape is different from another one.
+		 * 
+		 * @param OtherShape the other symbolic tensor shape.
+		 * @return true if the symbolic tensor shape is different from the other one, false otherwise.
+		 */
+		bool operator!=(const FSymbolicTensorShape& OtherShape) const;
+
+		/**
+		 * Assign this symbolic tensor shape from another one.
+		 * 
+		 * @param OtherShape the other symbolic tensor shape to copy from.
+		 */
+		void operator=(const FSymbolicTensorShape& OtherShape);
+	};
+
+	/**
+	 * The concrete shape of a tensor.
+	 *
+	 * Concrete tensor shapes are well defined through strictly positive values and thus have a defined volume.
+	 */
+	class NNE_API FTensorShape
+	{
+	public:
+		/**
+		 * The maximum number of dimensions supported by the tensor shape, aka the rank.
+		 */
+		constexpr static int32	MaxRank = FSymbolicTensorShape::MaxRank;
+
+	private:
+		TArray<uint32, TInlineAllocator<MaxRank>> Data;
+
+	public:
+		/**
+		 * Default construct a tensor shape.
+		 */
+		FTensorShape() = default;
+
+		/**
+		 * Construct this tensor shape from another one.
+		 * @param OtherShape the other tensor shape to copy from.
+		 */
+		FTensorShape(const FTensorShape& OtherShape) = default;
+
+		/**
+		 * Construct a tensor shape with the given dimensions.
+		 *
+		 * @param Data an array of dimensions.
+		 * @return a tensor shape with the given dimensions.
+		 */
+		static FTensorShape Make(TConstArrayView<uint32> Data);
+		
+		/**
+		 * Construct a tensor shape from a symbolic tensor shape.
+		 * 
+		 * Negative dimension values are replaced by 1.
+		 *
+		 * @param SymbolicShape a symbolic tensor shape.
+		 * @return a tensor shape with the given dimensions.
+		 */
+		static FTensorShape MakeFromSymbolic(const FSymbolicTensorShape& SymbolicShape);
+
+		/**
+		 * Get the dimensions of the tensor shape.
+		 *
+		 * @return a view on the dimensions of the tensor shape.
+		 */
+		inline TConstArrayView<uint32> GetData() const { return Data; };
+		
+		/**
+		 * Get the number of dimensions of the tensor shape, aka the rank.
+		 *
+		 * @return the number of dimensions of the tensor shape.
+		 */
+		inline int32 Rank() const { return Data.Num(); }
+		
+		/**
+		 * Get the number of elements in this tensor shape.
+		 *
+		 * @return the number of element in this tensor shape.
+		 */
+		uint64 Volume() const;
+		
+		/**
+		 * Check if this tensor shape is part of the ensemble of shapes defined by the symbolic shape.
+		 * 
+		 * For example [1,2] is compatible with [-1,2] and [1,2] but not with [2,2] or [1,2,3].
+		 *
+		 * @return true if this tensor shape is compatible with the symbolic shape, false otherwise.
+		 */
+		bool IsCompatibleWith(const FSymbolicTensorShape& SymbolicShape) const;
+
+		/**
+		 * Check if the tensor shape is equal to another one.
+		 * 
+		 * @param OtherShape the other tensor shape.
+		 * @return true if the tensor shape is equal to the other one, false otherwise.
+		 */
+		bool operator==(const FTensorShape& OtherShape) const;
+		
+		/**
+		 * Check if the tensor shape is different from another one.
+		 * 
+		 * @param OtherShape the other tensor shape.
+		 * @return true if the tensor shape is different from the other one, false otherwise.
+		 */
+		bool operator!=(const FTensorShape& OtherShape) const;
+
+		/**
+		 * Assign this tensor shape from another one.
+		 * @param OtherShape the other tensor shape to copy from.
+		 */
+		void operator=(const FTensorShape& OtherShape);
+	};
+
+	/**
+	 * Return the data size in bytes for a tensor data type.
+	 * 
+	 * @param InType the type of the element to consider.
+	 * @return the data size in bytes of an element.
+	 */
+	size_t NNE_API GetTensorDataTypeSizeInBytes(ENNETensorDataType InType);
+
+	/**
+	 * The descriptor for a tensor as model inputs and output.
+	 *
+	 * A tensor is described by its name, the type of data it contains and it's shape.
+	 * Since input and output tensors of a neural network can have dynamic shapes, Shape is symbolic.
+	 */
+	class NNE_API FTensorDesc
+	{
+		FString					Name;
+		ENNETensorDataType		DataType;
+		FSymbolicTensorShape	Shape;
+
+		FTensorDesc() = default;
+
+	public:
+		/**
+		 * Construct a tensor description for model input/output
+		 *
+		 * @param Name the name of the tensor.
+		 * @param Shape the symbolic shape of the tensor.
+		 * @param DataType the type of data the tensor contains.
+		 * @return a tensor shape with the given dimensions.
+		 */
+		static FTensorDesc Make(const FString& Name, const FSymbolicTensorShape& Shape, ENNETensorDataType DataType);
+		/**
+		 * Get the name of the tensor
+		 * 
+		 * @return the name of the tensor.
+		 */
+		inline const FString& GetName() const { return Name; }
+		
+		/**
+		 * Get the data type of the tensor
+		 * 
+		 * @return the data type of the tensor.
+		 */
+		inline ENNETensorDataType GetDataType() const { return DataType; }
+		
+		/**
+		 * Get the size in bytes of one element of the tensor.
+		 * 
+		 * @return the size in bytes of one element of the tensor.
+		 */
+		inline uint32 GetElemByteSize() const { return GetTensorDataTypeSizeInBytes(DataType); }
+
+		/**
+		 * Get the symbolic shape of the tensor.
+		 * 
+		 * @return the symbolic shape of the tensor.
+		 */
+		inline const FSymbolicTensorShape& GetShape() const { return Shape; }
+	};
+
+} // namespace UE::NNE
