NNE: UNNEModelData expose ModelData as ref counted memory.

Problem:
UNNEModelData and IModel implementation might contain the same data, increasing the peak memory when a Model is created until UNNEModelData is unloaded.

Solution:
Make UNNEModelData expose per runtime GetModelData() as ref counted allowing runtimes to refer to it even after UNNEModelData is unloaded.
+ Apply the pattern to all nne IModel implementations.

Tests:
win64 editor: Style transfer demo + all NNE tests.
win64 standalone: Style transfer demo.

#rnx

diff --git a/Engine/Plugins/Experimental/NNERuntimeORTCpu/Source/NNERuntimeORTCpu/Private/NNERuntimeORTCpu.cpp b/Engine/Plugins/Experimental/NNERuntimeORTCpu/Source/NNERuntimeORTCpu/Private/NNERuntimeORTCpu.cpp
--- a/Engine/Plugins/Experimental/NNERuntimeORTCpu/Source/NNERuntimeORTCpu/Private/NNERuntimeORTCpu.cpp
+++ b/Engine/Plugins/Experimental/NNERuntimeORTCpu/Source/NNERuntimeORTCpu/Private/NNERuntimeORTCpu.cpp
@@ -61,7 +61,14 @@ bool UNNERuntimeORTCpuImpl::CanCreateModelCPU(TObjectPtr<UNNEModelData> ModelDat
 	
 	int32 GuidSize = sizeof(UNNERuntimeORTCpuImpl::GUID);
 	int32 VersionSize = sizeof(UNNERuntimeORTCpuImpl::Version);
-	TConstArrayView<uint8> Data = ModelData->GetModelData(GetRuntimeName());
+	TSharedPtr<UE::NNE::FSharedModelData> SharedData = ModelData->GetModelData(GetRuntimeName());
+
+	if (!SharedData.IsValid())
+	{
+		return false;
+	}
+
+	TConstArrayView<uint8> Data = SharedData->GetView();
 	
 	if (Data.Num() <= GuidSize + VersionSize)
 	{
@@ -82,7 +89,8 @@ TSharedPtr<UE::NNE::IModelCPU> UNNERuntimeORTCpuImpl::CreateModelCPU(TObjectPtr<
 		return TSharedPtr<UE::NNE::IModelCPU>();
 	}
 
-	TConstArrayView<uint8> Data = ModelData->GetModelData(GetRuntimeName());
+	TSharedPtr<UE::NNE::FSharedModelData> Data = ModelData->GetModelData(GetRuntimeName());
+	check(Data.IsValid());
 	UE::NNERuntimeORTCpu::Private::FModelCPU* Model = new UE::NNERuntimeORTCpu::Private::FModelCPU(&NNEEnvironmentCPU, Data);
 	UE::NNE::IModelCPU* IModel = static_cast<UE::NNE::IModelCPU*>(Model);
 
@@ -91,7 +99,7 @@ TSharedPtr<UE::NNE::IModelCPU> UNNERuntimeORTCpuImpl::CreateModelCPU(TObjectPtr<
 		TArray<FAnalyticsEventAttribute> Attributes = MakeAnalyticsEventAttributeArray(
 			TEXT("PlatformName"), UGameplayStatics::GetPlatformName(),
 			TEXT("HashedRuntimeName"), FMD5::HashAnsiString(*GetRuntimeName()),
-			TEXT("ModelDataSize"), Data.Num()
+			TEXT("ModelDataSize"), Data->GetView().Num()
 		);
 		FEngineAnalytics::GetProvider().RecordEvent(TEXT("NeuralNetworkEngine.CreateModel"), Attributes);
 	}
diff --git a/Engine/Plugins/Experimental/NNERuntimeORTCpu/Source/NNERuntimeORTCpu/Private/NNERuntimeORTCpuModel.cpp b/Engine/Plugins/Experimental/NNERuntimeORTCpu/Source/NNERuntimeORTCpu/Private/NNERuntimeORTCpuModel.cpp
--- a/Engine/Plugins/Experimental/NNERuntimeORTCpu/Source/NNERuntimeORTCpu/Private/NNERuntimeORTCpuModel.cpp
+++ b/Engine/Plugins/Experimental/NNERuntimeORTCpu/Source/NNERuntimeORTCpu/Private/NNERuntimeORTCpuModel.cpp
@@ -300,22 +300,23 @@ namespace UE::NNERuntimeORTCpu::Private
 		InputTransferStatisticsEstimator.ResetStats();
 	}
 
-	TSharedPtr<UE::NNE::IModelInstanceCPU> FModelCPU::CreateModelInstanceCPU()
+	TSharedPtr<NNE::IModelInstanceCPU> FModelCPU::CreateModelInstanceCPU()
 	{
-		const UE::NNERuntimeORTCpu::Private::FRuntimeConf InConf;
-		UE::NNERuntimeORTCpu::Private::FModelInstanceCPU* ModelInstance = new UE::NNERuntimeORTCpu::Private::FModelInstanceCPU(ORTEnvironment, InConf);
+		const NNERuntimeORTCpu::Private::FRuntimeConf InConf;
+		NNERuntimeORTCpu::Private::FModelInstanceCPU* ModelInstance = new NNERuntimeORTCpu::Private::FModelInstanceCPU(ORTEnvironment, InConf);
 
-		if (!ModelInstance->Init(ModelData))
+		check(ModelData.IsValid());
+		if (!ModelInstance->Init(ModelData->GetView()))
 		{
 			delete ModelInstance;
-			return TSharedPtr<UE::NNE::IModelInstanceCPU>();
+			return TSharedPtr<NNE::IModelInstanceCPU>();
 		}
 
-		UE::NNE::IModelInstanceCPU* IModelInstance = static_cast<UE::NNE::IModelInstanceCPU*>(ModelInstance);
-		return TSharedPtr<UE::NNE::IModelInstanceCPU>(IModelInstance);
+		NNE::IModelInstanceCPU* IModelInstance = static_cast<NNE::IModelInstanceCPU*>(ModelInstance);
+		return TSharedPtr<NNE::IModelInstanceCPU>(IModelInstance);
 	}
 
-	FModelCPU::FModelCPU(Ort::Env* InORTEnvironment, TConstArrayView<uint8> InModelData) :
+	FModelCPU::FModelCPU(Ort::Env* InORTEnvironment, const TSharedPtr<UE::NNE::FSharedModelData>& InModelData) :
 		ORTEnvironment(InORTEnvironment),
 		ModelData(InModelData)
 	{
diff --git a/Engine/Plugins/Experimental/NNERuntimeORTCpu/Source/NNERuntimeORTCpu/Private/NNERuntimeORTCpuModel.h b/Engine/Plugins/Experimental/NNERuntimeORTCpu/Source/NNERuntimeORTCpu/Private/NNERuntimeORTCpuModel.h
--- a/Engine/Plugins/Experimental/NNERuntimeORTCpu/Source/NNERuntimeORTCpu/Private/NNERuntimeORTCpuModel.h
+++ b/Engine/Plugins/Experimental/NNERuntimeORTCpu/Source/NNERuntimeORTCpu/Private/NNERuntimeORTCpuModel.h
@@ -4,6 +4,7 @@
 
 #include "CoreMinimal.h"
 #include "NNEModelBase.h"
+#include "NNEModelData.h"
 #include "NNERuntimeCPU.h"
 #include "NNETensor.h"
 #include "NNETypes.h"
@@ -83,14 +84,14 @@ namespace UE::NNERuntimeORTCpu::Private
 	class FModelCPU : public NNE::IModelCPU
 	{
 	public:
-		FModelCPU(Ort::Env* InORTEnvironment, TConstArrayView<uint8> ModelData);
+		FModelCPU(Ort::Env* InORTEnvironment, const TSharedPtr<UE::NNE::FSharedModelData>& InModelData);
 		virtual ~FModelCPU() {};
 
-		virtual TSharedPtr<UE::NNE::IModelInstanceCPU> CreateModelInstanceCPU() override;
+		virtual TSharedPtr<NNE::IModelInstanceCPU> CreateModelInstanceCPU() override;
 
 	private:
 		Ort::Env* ORTEnvironment;
-		TArray<uint8> ModelData;
+		TSharedPtr<UE::NNE::FSharedModelData> ModelData;
 	};
 	
 } // UE::NNERuntimeORTCpu::Private
\ No newline at end of file
diff --git a/Engine/Plugins/Experimental/NNERuntimeORTGpu/Source/NNERuntimeORT/Private/NNERuntimeORT.cpp b/Engine/Plugins/Experimental/NNERuntimeORTGpu/Source/NNERuntimeORT/Private/NNERuntimeORT.cpp
--- a/Engine/Plugins/Experimental/NNERuntimeORTGpu/Source/NNERuntimeORT/Private/NNERuntimeORT.cpp
+++ b/Engine/Plugins/Experimental/NNERuntimeORTGpu/Source/NNERuntimeORT/Private/NNERuntimeORT.cpp
@@ -82,7 +82,14 @@ bool UNNERuntimeORTGpuImpl::CanCreateModelGPU(TObjectPtr<UNNEModelData> ModelDat
 
 	int32 GuidSize = sizeof(UNNERuntimeORTGpuImpl::GUID);
 	int32 VersionSize = sizeof(UNNERuntimeORTGpuImpl::Version);
-	TConstArrayView<uint8> Data = ModelData->GetModelData(GetRuntimeName());
+	TSharedPtr<UE::NNE::FSharedModelData> SharedData = ModelData->GetModelData(GetRuntimeName());
+
+	if (!SharedData.IsValid())
+	{
+		return false;
+	}
+
+	TConstArrayView<uint8> Data = SharedData->GetView();
 
 	if (Data.Num() <= GuidSize + VersionSize)
 	{
@@ -109,7 +116,8 @@ TSharedPtr<UE::NNE::IModelGPU> UNNERuntimeORTGpuImpl::CreateModelGPU(TObjectPtr<
 	}
 
 	UE::NNE::IModelGPU* IModel = nullptr;
-	TConstArrayView<uint8> Data = ModelData->GetModelData(GetRuntimeName());
+	TSharedPtr<UE::NNE::FSharedModelData> Data = ModelData->GetModelData(GetRuntimeName());
+	check(Data.IsValid());
 
 	switch (Provider)
 	{
@@ -129,7 +137,7 @@ TSharedPtr<UE::NNE::IModelGPU> UNNERuntimeORTGpuImpl::CreateModelGPU(TObjectPtr<
 		TArray<FAnalyticsEventAttribute> Attributes = MakeAnalyticsEventAttributeArray(
 			TEXT("PlatformName"), UGameplayStatics::GetPlatformName(),
 			TEXT("HashedRuntimeName"), FMD5::HashAnsiString(*GetRuntimeName()),
-			TEXT("ModelDataSize"), Data.Num()
+			TEXT("ModelDataSize"), Data->GetView().Num()
 		);
 		FEngineAnalytics::GetProvider().RecordEvent(TEXT("NeuralNetworkEngine.CreateModel"), Attributes);
 	}
diff --git a/Engine/Plugins/Experimental/NNERuntimeORTGpu/Source/NNERuntimeORT/Private/NNERuntimeORTModel.cpp b/Engine/Plugins/Experimental/NNERuntimeORTGpu/Source/NNERuntimeORT/Private/NNERuntimeORTModel.cpp
--- a/Engine/Plugins/Experimental/NNERuntimeORTGpu/Source/NNERuntimeORT/Private/NNERuntimeORTModel.cpp
+++ b/Engine/Plugins/Experimental/NNERuntimeORTGpu/Source/NNERuntimeORT/Private/NNERuntimeORTModel.cpp
@@ -327,44 +327,46 @@ namespace UE::NNERuntimeORT::Private
 
 #if PLATFORM_WINDOWS
 
-	TSharedPtr<UE::NNE::IModelInstanceGPU> FModelORTDml::CreateModelInstanceGPU()
+	TSharedPtr<NNE::IModelInstanceGPU> FModelORTDml::CreateModelInstanceGPU()
 	{
 		const FRuntimeConf InConf;
 		FModelInstanceORTDml* ModelInstance = new FModelInstanceORTDml(ORTEnvironment, InConf);
 
-		if (!ModelInstance->Init(ModelData))
+		check(ModelData.IsValid());
+		if (!ModelInstance->Init(ModelData->GetView()))
 		{
 			delete ModelInstance;
 			return TSharedPtr<UE::NNE::IModelInstanceGPU>();
 		}
 
-		UE::NNE::IModelInstanceGPU* IModelInstance = static_cast<UE::NNE::IModelInstanceGPU*>(ModelInstance);
-		return TSharedPtr<UE::NNE::IModelInstanceGPU>(IModelInstance);
+		NNE::IModelInstanceGPU* IModelInstance = static_cast<NNE::IModelInstanceGPU*>(ModelInstance);
+		return TSharedPtr<NNE::IModelInstanceGPU>(IModelInstance);
 	}
 
-	FModelORTDml::FModelORTDml(Ort::Env* InORTEnvironment, TConstArrayView<uint8> InModelData) :
+	FModelORTDml::FModelORTDml(Ort::Env* InORTEnvironment, const TSharedPtr<UE::NNE::FSharedModelData>& InModelData) :
 		ORTEnvironment(InORTEnvironment),
 		ModelData(InModelData)
 	{
 
 	}
 
-	TSharedPtr<UE::NNE::IModelInstanceGPU> FModelORTCuda::CreateModelInstanceGPU()
+	TSharedPtr<NNE::IModelInstanceGPU> FModelORTCuda::CreateModelInstanceGPU()
 	{
 		const FRuntimeConf InConf;
 		FModelInstanceORTCuda* ModelInstance = new FModelInstanceORTCuda(ORTEnvironment, InConf);
 
-		if (!ModelInstance->Init(ModelData))
+		check(ModelData.IsValid());
+		if (!ModelInstance->Init(ModelData->GetView()))
 		{
 			delete ModelInstance;
-			return TSharedPtr<UE::NNE::IModelInstanceGPU>();
+			return TSharedPtr<NNE::IModelInstanceGPU>();
 		}
 
-		UE::NNE::IModelInstanceGPU* IModelInstance = static_cast<UE::NNE::IModelInstanceGPU*>(ModelInstance);
-		return TSharedPtr<UE::NNE::IModelInstanceGPU>(IModelInstance);
+		NNE::IModelInstanceGPU* IModelInstance = static_cast<NNE::IModelInstanceGPU*>(ModelInstance);
+		return TSharedPtr<NNE::IModelInstanceGPU>(IModelInstance);
 	}
 
-	FModelORTCuda::FModelORTCuda(Ort::Env* InORTEnvironment, TConstArrayView<uint8> InModelData) :
+	FModelORTCuda::FModelORTCuda(Ort::Env* InORTEnvironment, const TSharedPtr<UE::NNE::FSharedModelData>& InModelData) :
 		ORTEnvironment(InORTEnvironment),
 		ModelData(InModelData)
 	{
diff --git a/Engine/Plugins/Experimental/NNERuntimeORTGpu/Source/NNERuntimeORT/Private/NNERuntimeORTModel.h b/Engine/Plugins/Experimental/NNERuntimeORTGpu/Source/NNERuntimeORT/Private/NNERuntimeORTModel.h
--- a/Engine/Plugins/Experimental/NNERuntimeORTGpu/Source/NNERuntimeORT/Private/NNERuntimeORTModel.h
+++ b/Engine/Plugins/Experimental/NNERuntimeORTGpu/Source/NNERuntimeORT/Private/NNERuntimeORTModel.h
@@ -4,6 +4,7 @@
 
 #include "CoreMinimal.h"
 #include "NNEModelBase.h"
+#include "NNEModelData.h"
 #include "NNERuntimeGPU.h"
 #include "NNETensor.h"
 #include "NNETypes.h"
@@ -104,27 +105,27 @@ namespace UE::NNERuntimeORT::Private
 	class FModelORTDml : public NNE::IModelGPU
 	{
 	public:
-		FModelORTDml(Ort::Env* InORTEnvironment, TConstArrayView<uint8> ModelData);
+		FModelORTDml(Ort::Env* InORTEnvironment, const TSharedPtr<UE::NNE::FSharedModelData>& InModelData);
 		virtual ~FModelORTDml() {};
 
-		virtual TSharedPtr<UE::NNE::IModelInstanceGPU> CreateModelInstanceGPU() override;
+		virtual TSharedPtr<NNE::IModelInstanceGPU> CreateModelInstanceGPU() override;
 
 	private:
 		Ort::Env* ORTEnvironment;
-		TArray<uint8> ModelData;
+		TSharedPtr<UE::NNE::FSharedModelData> ModelData;
 	};
 
 	class FModelORTCuda : public NNE::IModelGPU
 	{
 	public:
-		FModelORTCuda(Ort::Env* InORTEnvironment, TConstArrayView<uint8> ModelData);
+		FModelORTCuda(Ort::Env* InORTEnvironment, const TSharedPtr<UE::NNE::FSharedModelData>& InModelData);
 		virtual ~FModelORTCuda() {};
 
-		virtual TSharedPtr<UE::NNE::IModelInstanceGPU> CreateModelInstanceGPU() override;
+		virtual TSharedPtr<NNE::IModelInstanceGPU> CreateModelInstanceGPU() override;
 
 	private:
 		Ort::Env* ORTEnvironment;
-		TArray<uint8> ModelData;
+		TSharedPtr<UE::NNE::FSharedModelData> ModelData;
 	};
 #endif
 	
diff --git a/Engine/Plugins/Experimental/NNERuntimeRDG/Source/NNERuntimeRDG/Private/Dml/NNEDmlModel.cpp b/Engine/Plugins/Experimental/NNERuntimeRDG/Source/NNERuntimeRDG/Private/Dml/NNEDmlModel.cpp
--- a/Engine/Plugins/Experimental/NNERuntimeRDG/Source/NNERuntimeRDG/Private/Dml/NNEDmlModel.cpp
+++ b/Engine/Plugins/Experimental/NNERuntimeRDG/Source/NNERuntimeRDG/Private/Dml/NNEDmlModel.cpp
@@ -1386,21 +1386,22 @@ int FModelInstance::PrepareTensorShapesAndData()
 	return 0;
 }
 
-TSharedPtr<UE::NNE::IModelInstanceRDG> FModel::CreateModelInstanceRDG()
+TSharedPtr<NNE::IModelInstanceRDG> FModel::CreateModelInstanceRDG()
 {
 	FModelInstance* ModelInstance = new FModelInstance();
 
-	if (!ModelInstance->Init(ModelData, DevCtx))
+	check(ModelData.IsValid());
+	if (!ModelInstance->Init(ModelData->GetView(), DevCtx))
 	{
 		delete ModelInstance;
-		return TSharedPtr<UE::NNE::IModelInstanceRDG>();
+		return TSharedPtr<NNE::IModelInstanceRDG>();
 	}
 
-	UE::NNE::IModelInstanceRDG* IModelInstance = static_cast<UE::NNE::IModelInstanceRDG*>(ModelInstance);
-	return TSharedPtr<UE::NNE::IModelInstanceRDG>(IModelInstance);
+	NNE::IModelInstanceRDG* IModelInstance = static_cast<NNE::IModelInstanceRDG*>(ModelInstance);
+	return TSharedPtr<NNE::IModelInstanceRDG>(IModelInstance);
 }
 
-FModel::FModel(TConstArrayView<uint8> InModelData, FDmlDeviceContext* InDevCtx) 
+FModel::FModel(const TSharedPtr<NNE::FSharedModelData>& InModelData, FDmlDeviceContext* InDevCtx)
 	: ModelData(InModelData), DevCtx(InDevCtx)
 {
 }
diff --git a/Engine/Plugins/Experimental/NNERuntimeRDG/Source/NNERuntimeRDG/Private/Dml/NNEDmlModel.h b/Engine/Plugins/Experimental/NNERuntimeRDG/Source/NNERuntimeRDG/Private/Dml/NNEDmlModel.h
--- a/Engine/Plugins/Experimental/NNERuntimeRDG/Source/NNERuntimeRDG/Private/Dml/NNEDmlModel.h
+++ b/Engine/Plugins/Experimental/NNERuntimeRDG/Source/NNERuntimeRDG/Private/Dml/NNEDmlModel.h
@@ -3,6 +3,7 @@
 #pragma once
 
 #include "NNEDmlCommon.h"
+#include "NNEModelData.h"
 #include "NNERuntimeRDGModel.h"
 
 struct ID3D12DynamicRHI;
@@ -95,13 +96,13 @@ private:
 class FModel : public NNE::IModelRDG
 {
 public:
-	FModel(TConstArrayView<uint8> InModelData, FDmlDeviceContext* InDevCtx);
+	FModel(const TSharedPtr<NNE::FSharedModelData>& InModelData, FDmlDeviceContext* InDevCtx);
 	virtual ~FModel() {};
 
-	virtual TSharedPtr<UE::NNE::IModelInstanceRDG> CreateModelInstanceRDG() override;
+	virtual TSharedPtr<NNE::IModelInstanceRDG> CreateModelInstanceRDG() override;
 
 private:
-	TArray<uint8> ModelData;
+	TSharedPtr<NNE::FSharedModelData> ModelData;
 	FDmlDeviceContext* DevCtx;
 };
 
diff --git a/Engine/Plugins/Experimental/NNERuntimeRDG/Source/NNERuntimeRDG/Private/NNERuntimeRDGDml.cpp b/Engine/Plugins/Experimental/NNERuntimeRDG/Source/NNERuntimeRDG/Private/NNERuntimeRDGDml.cpp
--- a/Engine/Plugins/Experimental/NNERuntimeRDG/Source/NNERuntimeRDG/Private/NNERuntimeRDGDml.cpp
+++ b/Engine/Plugins/Experimental/NNERuntimeRDG/Source/NNERuntimeRDG/Private/NNERuntimeRDGDml.cpp
@@ -143,7 +143,14 @@ FString UNNERuntimeRDGDmlImpl::GetModelDataIdentifier(FString FileType, TConstAr
 bool UNNERuntimeRDGDmlImpl::CanCreateModelRDG(TObjectPtr<UNNEModelData> ModelData) const
 {
 #ifdef NNE_USE_DIRECTML
-	TConstArrayView<uint8> Data = ModelData->GetModelData(GetRuntimeName());
+	TSharedPtr<UE::NNE::FSharedModelData> SharedData = ModelData->GetModelData(GetRuntimeName());
+
+	if (!SharedData.IsValid())
+	{
+		return false;
+	}
+
+	TConstArrayView<uint8> Data = SharedData->GetView();
 
 	if (Data.Num() <= FModelInfo::Get()->GetGuidSize() + FModelInfo::Get()->GetVersionSize())
 	{
@@ -165,7 +172,8 @@ TSharedPtr<UE::NNE::IModelRDG> UNNERuntimeRDGDmlImpl::CreateModelRDG(TObjectPtr<
 		return TSharedPtr<UE::NNE::IModelRDG>();
 	}
 
-	TConstArrayView<uint8> Data = ModelData->GetModelData(GetRuntimeName());
+	TSharedPtr<UE::NNE::FSharedModelData> Data = ModelData->GetModelData(GetRuntimeName());
+	check(Data.IsValid());
 	FModel* Model = new FModel(Data, Ctx);
 
 	if (FEngineAnalytics::IsAvailable())
@@ -173,7 +181,7 @@ TSharedPtr<UE::NNE::IModelRDG> UNNERuntimeRDGDmlImpl::CreateModelRDG(TObjectPtr<
 		TArray<FAnalyticsEventAttribute> Attributes = MakeAnalyticsEventAttributeArray(
 			TEXT("PlatformName"), UGameplayStatics::GetPlatformName(),
 			TEXT("HashedRuntimeName"), FMD5::HashAnsiString(*GetRuntimeName()),
-			TEXT("ModelDataSize"), Data.Num()
+			TEXT("ModelDataSize"), Data->GetView().Num()
 		);
 		FEngineAnalytics::GetProvider().RecordEvent(TEXT("NeuralNetworkEngine.CreateModel"), Attributes);
 	}
diff --git a/Engine/Plugins/Experimental/NNERuntimeRDG/Source/NNERuntimeRDG/Private/NNERuntimeRDGHlsl.cpp b/Engine/Plugins/Experimental/NNERuntimeRDG/Source/NNERuntimeRDG/Private/NNERuntimeRDGHlsl.cpp
--- a/Engine/Plugins/Experimental/NNERuntimeRDG/Source/NNERuntimeRDG/Private/NNERuntimeRDGHlsl.cpp
+++ b/Engine/Plugins/Experimental/NNERuntimeRDG/Source/NNERuntimeRDG/Private/NNERuntimeRDGHlsl.cpp
@@ -82,8 +82,15 @@ bool UNNERuntimeRDGHlslImpl::CanCreateModelRDG(TObjectPtr<UNNEModelData> ModelDa
 {
 	int32 GuidSize = sizeof(GUID);
 	int32 VersionSize = sizeof(Version);
-	TConstArrayView<uint8> Data = ModelData->GetModelData(GetRuntimeName());
-	
+	TSharedPtr<UE::NNE::FSharedModelData> SharedData = ModelData->GetModelData(GetRuntimeName());
+
+	if (!SharedData.IsValid())
+	{
+		return false;
+	}
+
+	TConstArrayView<uint8> Data = SharedData->GetView();
+
 	if (Data.Num() <= GuidSize + VersionSize)
 	{
 		return false;
@@ -134,7 +141,8 @@ TSharedPtr<UE::NNE::IModelRDG> UNNERuntimeRDGHlslImpl::CreateModelRDG(TObjectPtr
 		return TSharedPtr<UE::NNE::IModelRDG>();
 	}
 
-	TConstArrayView<uint8> Data = ModelData->GetModelData(GetRuntimeName());
+	TSharedPtr<UE::NNE::FSharedModelData> Data = ModelData->GetModelData(GetRuntimeName());
+	check(Data.IsValid());
 	UE::NNERuntimeRDG::Private::Hlsl::FModel* Model = new UE::NNERuntimeRDG::Private::Hlsl::FModel(Data);
 
 	if (FEngineAnalytics::IsAvailable())
@@ -142,7 +150,7 @@ TSharedPtr<UE::NNE::IModelRDG> UNNERuntimeRDGHlslImpl::CreateModelRDG(TObjectPtr
 		TArray<FAnalyticsEventAttribute> Attributes = MakeAnalyticsEventAttributeArray(
 			TEXT("PlatformName"), UGameplayStatics::GetPlatformName(),
 			TEXT("HashedRuntimeName"), FMD5::HashAnsiString(*GetRuntimeName()),
-			TEXT("ModelDataSize"), Data.Num()
+			TEXT("ModelDataSize"), Data->GetView().Num()
 		);
 		FEngineAnalytics::GetProvider().RecordEvent(TEXT("NeuralNetworkEngine.CreateModel"), Attributes);
 	}
diff --git a/Engine/Plugins/Experimental/NNERuntimeRDG/Source/NNERuntimeRDG/Private/NNERuntimeRDGModelHlsl.cpp b/Engine/Plugins/Experimental/NNERuntimeRDG/Source/NNERuntimeRDG/Private/NNERuntimeRDGModelHlsl.cpp
--- a/Engine/Plugins/Experimental/NNERuntimeRDG/Source/NNERuntimeRDG/Private/NNERuntimeRDGModelHlsl.cpp
+++ b/Engine/Plugins/Experimental/NNERuntimeRDG/Source/NNERuntimeRDG/Private/NNERuntimeRDGModelHlsl.cpp
@@ -364,20 +364,21 @@ bool FModelInstance::PrepareWeights()
 	return true;
 }
 
-TSharedPtr<UE::NNE::IModelInstanceRDG> FModel::CreateModelInstanceRDG()
+TSharedPtr<NNE::IModelInstanceRDG> FModel::CreateModelInstanceRDG()
 {
 	FModelInstance* ModelInstance = new FModelInstance();
 
-	if (!ModelInstance->Init(ModelData))
+	check(ModelData.IsValid());
+	if (!ModelInstance->Init(ModelData->GetView()))
 	{
 		delete ModelInstance;
-		return TSharedPtr<UE::NNE::IModelInstanceRDG>();
+		return TSharedPtr<NNE::IModelInstanceRDG>();
 	}
 
-	UE::NNE::IModelInstanceRDG* IModelInstance = static_cast<UE::NNE::IModelInstanceRDG*>(ModelInstance);
-	return TSharedPtr<UE::NNE::IModelInstanceRDG>(IModelInstance);
+	NNE::IModelInstanceRDG* IModelInstance = static_cast<NNE::IModelInstanceRDG*>(ModelInstance);
+	return TSharedPtr<NNE::IModelInstanceRDG>(IModelInstance);
 }
 
-FModel::FModel(TConstArrayView<uint8> InModelData) : ModelData(InModelData) {}
+FModel::FModel(const TSharedPtr<UE::NNE::FSharedModelData>& InModelData) : ModelData(InModelData) {}
 
 } // namespace UE::NNERuntimeRDG::Private::Hlsl
\ No newline at end of file
diff --git a/Engine/Plugins/Experimental/NNERuntimeRDG/Source/NNERuntimeRDG/Private/NNERuntimeRDGModelHlsl.h b/Engine/Plugins/Experimental/NNERuntimeRDG/Source/NNERuntimeRDG/Private/NNERuntimeRDGModelHlsl.h
--- a/Engine/Plugins/Experimental/NNERuntimeRDG/Source/NNERuntimeRDG/Private/NNERuntimeRDGModelHlsl.h
+++ b/Engine/Plugins/Experimental/NNERuntimeRDG/Source/NNERuntimeRDG/Private/NNERuntimeRDGModelHlsl.h
@@ -3,6 +3,7 @@
 #pragma once
 
 #include "NNERuntimeRDGModel.h"
+#include "NNEModelData.h"
 #include "NNETypes.h"
 
 class FRDGBuilder;
@@ -38,13 +39,13 @@ private:
 class FModel : public NNE::IModelRDG
 {
 public:
-	FModel(TConstArrayView<uint8> ModelData);
+	FModel(const TSharedPtr<UE::NNE::FSharedModelData>& InModelData);
 	virtual ~FModel() {};
 
-	virtual TSharedPtr<UE::NNE::IModelInstanceRDG> CreateModelInstanceRDG() override;
+	virtual TSharedPtr<NNE::IModelInstanceRDG> CreateModelInstanceRDG() override;
 
 private:
-	TArray<uint8> ModelData;
+	TSharedPtr<UE::NNE::FSharedModelData> ModelData;
 };
 
 } // namespace UE::NNERuntimeRDG::Private::Hlsl
\ No newline at end of file
diff --git a/Engine/Source/Runtime/NNE/Private/NNEModelData.cpp b/Engine/Source/Runtime/NNE/Private/NNEModelData.cpp
--- a/Engine/Source/Runtime/NNE/Private/NNEModelData.cpp
+++ b/Engine/Source/Runtime/NNE/Private/NNEModelData.cpp
@@ -16,7 +16,6 @@
 #include "DerivedDataCache.h"
 #include "DerivedDataRequestOwner.h"
 #include "Internationalization/TextLocalizationResource.h"
-#include "Memory/SharedBuffer.h"
 #include "Misc/Guid.h"
 #endif
 
@@ -122,7 +121,7 @@ FGuid UNNEModelData::GetFileId()
 	return FileId;
 }
 
-TConstArrayView<uint8> UNNEModelData::GetModelData(const FString& RuntimeName)
+TSharedPtr<UE::NNE::FSharedModelData> UNNEModelData::GetModelData(const FString& RuntimeName)
 {
 #if WITH_EDITORONLY_DATA
 	// Check model data is supporting the requested target runtime
@@ -134,15 +133,15 @@ TConstArrayView<uint8> UNNEModelData::GetModelData(const FString& RuntimeName)
 		{
 			UE_LOG(LogNNE, Error, TEXT("- %s"), *TargetRuntimesName);
 		}
-		return {};
+		return TSharedPtr<UE::NNE::FSharedModelData>();
 	}
 #endif //WITH_EDITORONLY_DATA
 
 	// Check if we have a local cache hit
-	TArray<uint8>* LocalData = ModelData.Find(RuntimeName);
+	FSharedBuffer* LocalData = ModelData.Find(RuntimeName);
 	if (LocalData)
 	{
-		return TConstArrayView<uint8>(LocalData->GetData(), LocalData->Num());
+		return MakeShared<UE::NNE::FSharedModelData>(*LocalData);
 	}
 	
 #if WITH_EDITOR
@@ -150,32 +149,33 @@ TConstArrayView<uint8> UNNEModelData::GetModelData(const FString& RuntimeName)
 	if (!NNERuntime.IsValid())
 	{
 		UE_LOG(LogNNE, Error, TEXT("UNNEModelData: Runtime '%s' is among the target runtimes but instance is invalid."), *RuntimeName);
-		return {};
+		return TSharedPtr<UE::NNE::FSharedModelData>();
 	}
 
 	FString ModelDataIdentifier = NNERuntime->GetModelDataIdentifier(FileType, FileData, FileId, nullptr);
 	if (ModelDataIdentifier.Len() == 0)
 	{
 		UE_LOG(LogNNE, Error, TEXT("UNNEModelData: Runtime '%s' returned an empty string as a ModelDataIdentifier. GetModelDataIdentifier should always return a valid identifier."), *RuntimeName);
-		return {};
+		return TSharedPtr<UE::NNE::FSharedModelData>();
 	}
 	
 	// Check if we have a DDC cache hit
 	FSharedBuffer RemoteData = GetFromDDC(FileId, RuntimeName, ModelDataIdentifier);
 	if (RemoteData.GetSize() > 0)
 	{
-		ModelData.Add(RuntimeName, TArray<uint8>((uint8*)RemoteData.GetData(), RemoteData.GetSize()));
+		ModelData.Add(RuntimeName, RemoteData);
 		
-		TArray<uint8>* CachedRemoteData = ModelData.Find(RuntimeName);
-		return TConstArrayView<uint8>(CachedRemoteData->GetData(), CachedRemoteData->Num());
+		return MakeShared<UE::NNE::FSharedModelData>(RemoteData);
 	}
 #endif //WITH_EDITOR
 
 	// Try to create the model
-	TArray<uint8> CreatedData = CreateRuntimeDataBlob(RuntimeName, FileType, FileData, FileId, nullptr);
-	if (CreatedData.Num() < 1)
+	TArray<uint8> RuntimeDataBlob = CreateRuntimeDataBlob(RuntimeName, FileType, FileData, FileId, nullptr);
+	FSharedBuffer CreatedData = MakeSharedBufferFromArray(MoveTemp(RuntimeDataBlob));
+	if (CreatedData.GetSize() < 1)
 	{
-		return {};
+		
+		return TSharedPtr<UE::NNE::FSharedModelData>();
 	}
 
 	// Cache the model
@@ -183,12 +183,10 @@ TConstArrayView<uint8> UNNEModelData::GetModelData(const FString& RuntimeName)
 
 #if WITH_EDITOR
 	// And put it into DDC
-	FSharedBuffer SharedBuffer = MakeSharedBufferFromArray(MoveTemp(CreatedData));
-	PutIntoDDC(FileId, RuntimeName, ModelDataIdentifier, SharedBuffer);
+	PutIntoDDC(FileId, RuntimeName, ModelDataIdentifier, CreatedData);
 #endif //WITH_EDITOR
 	
-	TArray<uint8>* CachedCreatedData = ModelData.Find(RuntimeName);
-	return TConstArrayView<uint8>(CachedCreatedData->GetData(), CachedCreatedData->Num());
+	return MakeShared<UE::NNE::FSharedModelData>(CreatedData);
 }
 
 void UNNEModelData::Serialize(FArchive& Ar)
@@ -216,9 +214,10 @@ void UNNEModelData::Serialize(FArchive& Ar)
 
 		for (const FString& RuntimeName : CookedRuntimeNames)
 		{
-			TArray<uint8> CreatedData = CreateRuntimeDataBlob(RuntimeName, FileType, FileData, FileId, Ar.GetArchiveState().CookingTarget());
-			if (CreatedData.Num() > 0)
+			TArray RuntimeDataBlob = CreateRuntimeDataBlob(RuntimeName, FileType, FileData, FileId, Ar.GetArchiveState().CookingTarget());
+			if (RuntimeDataBlob.Num() > 0)
 			{
+				FSharedBuffer CreatedData = MakeSharedBufferFromArray(MoveTemp(RuntimeDataBlob));
 				ModelData.Add(RuntimeName, CreatedData);
 #if WITH_EDITOR
 				TWeakInterfacePtr<INNERuntime> NNERuntime = UE::NNE::GetRuntime<INNERuntime>(RuntimeName);
@@ -227,8 +226,7 @@ void UNNEModelData::Serialize(FArchive& Ar)
 					FString ModelDataIdentifier = NNERuntime->GetModelDataIdentifier(FileType, FileData, FileId, Ar.GetArchiveState().CookingTarget());
 					if (ModelDataIdentifier.Len() > 0)
 					{
-						FSharedBuffer SharedBuffer = MakeSharedBufferFromArray(MoveTemp(CreatedData));
-						PutIntoDDC(FileId, RuntimeName, ModelDataIdentifier, SharedBuffer);
+						PutIntoDDC(FileId, RuntimeName, ModelDataIdentifier, CreatedData);
 					}
 					else
 					{
@@ -257,7 +255,9 @@ void UNNEModelData::Serialize(FArchive& Ar)
 		for (int i = 0; i < NumItems; i++)
 		{
 			Ar << RuntimeNames[i];
-			Ar << ModelData[RuntimeNames[i]];
+			FSharedBuffer Data = ModelData[RuntimeNames[i]];
+			TArray<uint8> DataToSerialize(static_cast<const uint8*>(Data.GetData()), Data.GetSize());
+			Ar << DataToSerialize;
 		}
 	}
 	else
@@ -291,7 +291,8 @@ void UNNEModelData::Serialize(FArchive& Ar)
 				Ar << Name;
 				TArray<uint8> Data;
 				Ar << Data;
-				ModelData.Add(Name, MoveTemp(Data));
+				FSharedBuffer SharedData = MakeSharedBufferFromArray(MoveTemp(Data));
+				ModelData.Add(Name, SharedData);
 			}
 		}
 	}
diff --git a/Engine/Source/Runtime/NNE/Public/NNEModelData.h b/Engine/Source/Runtime/NNE/Public/NNEModelData.h
--- a/Engine/Source/Runtime/NNE/Public/NNEModelData.h
+++ b/Engine/Source/Runtime/NNE/Public/NNEModelData.h
@@ -3,9 +3,33 @@
 #pragma once
 
 #include "CoreMinimal.h"
+#include "Memory/SharedBuffer.h"
+#include "Templates/SharedPointer.h"
 
 #include "NNEModelData.generated.h"
 
+namespace UE::NNE
+{
+	/**
+	 * This class present a ref counted view on an immutable memory buffer. 
+	 * Allowing runtime to reference result of GetModelData() even if they outlive UNNEModelData.
+	 */
+	class FSharedModelData
+	{
+	private:
+		FSharedBuffer Data;
+
+	public:
+		FSharedModelData(FSharedBuffer InData) : Data(InData) {}
+		FSharedModelData() {}
+
+		TConstArrayView<uint8> GetView() const
+		{
+			return MakeArrayView(static_cast<const uint8*>(Data.GetData()), Data.GetSize());
+		}
+	};
+}
+
 /**
  * This class represents assets that store neural network model data.
  *
@@ -63,9 +87,9 @@ public:
 	 * This function is used by runtimes when creating a model. In editor, the function will create the optimized model data with the passed runtime in case it has not been cached in the DCC yet. In game, the cooked data is accessed.
 	 *
 	 * @param RuntimeName The name of the runtime for which the data should be returned.
-	 * @return The optimized and runtime specific model data or an empty view in case of a failure.
+	 * @return The optimized and runtime specific model data or an invalid TSharedPtr in case of failure.
 	 */
-	TConstArrayView<uint8> GetModelData(const FString& RuntimeName);
+	TSharedPtr<UE::NNE::FSharedModelData> GetModelData(const FString& RuntimeName);
 
 	/**
 	 * Implements custom serialization of this asset.
@@ -131,6 +155,6 @@ private:
 	/**
 	 * The processed / optimized model data for the different runtimes.
 	 */
-	TMap<FString, TArray<uint8>> ModelData;
+	TMap<FString, FSharedBuffer> ModelData;
 
 };
