[NNE] Remove unused NNE public API parts

- Removed IsPlatformSupported as it is covered by CanCreateModelData
- Cleaned up minor things

#rnx

diff --git a/Engine/Plugins/AI/MLAdapter/Source/MLAdapter/Private/Agents/MLAdapterAgent_Inference.cpp b/Engine/Plugins/AI/MLAdapter/Source/MLAdapter/Private/Agents/MLAdapterAgent_Inference.cpp
--- a/Engine/Plugins/AI/MLAdapter/Source/MLAdapter/Private/Agents/MLAdapterAgent_Inference.cpp
+++ b/Engine/Plugins/AI/MLAdapter/Source/MLAdapter/Private/Agents/MLAdapterAgent_Inference.cpp
@@ -57,14 +57,14 @@ void UMLAdapterAgent_Inference::PostInitProperties()
 		return;
 	}
 
-	InputTensorSizeInBytes = InputTensorDescs[0].GetElemByteSize() * InputTensorShape.Volume();
+	InputTensorSizeInBytes = InputTensorDescs[0].GetElementByteSize() * InputTensorShape.Volume();
 	if (InputTensorSizeInBytes == 0)
 	{
 		UE_LOG(LogMLAdapter, Warning, TEXT("Input tensor should be of non-zero size."));
 		return;
 	}
 
-	OutputTensorSizeInBytes = OutputTensorDescs[0].GetElemByteSize() * OutputTensorShapes[0].Volume();
+	OutputTensorSizeInBytes = OutputTensorDescs[0].GetElementByteSize() * OutputTensorShapes[0].Volume();
 	if (OutputTensorSizeInBytes == 0)
 	{
 		UE_LOG(LogMLAdapter, Warning, TEXT("Output tensor should be of non-zero size."));
diff --git a/Engine/Plugins/Animation/MLDeformer/VertexDeltaModel/Source/VertexDeltaModel/Private/VertexDeltaModelInstance.cpp b/Engine/Plugins/Animation/MLDeformer/VertexDeltaModel/Source/VertexDeltaModel/Private/VertexDeltaModelInstance.cpp
--- a/Engine/Plugins/Animation/MLDeformer/VertexDeltaModel/Source/VertexDeltaModel/Private/VertexDeltaModelInstance.cpp
+++ b/Engine/Plugins/Animation/MLDeformer/VertexDeltaModel/Source/VertexDeltaModel/Private/VertexDeltaModelInstance.cpp
@@ -98,11 +98,11 @@ bool UVertexDeltaModelInstance::GetRDGVertexBufferDesc(TConstArrayView<UE::NNE::
 {
 	if (InOutputTensorDescs.Num() > 0)
 	{
-		const uint32 ElemByteSize = InOutputTensorDescs[0].GetElemByteSize();
+		const uint32 ElemByteSize = InOutputTensorDescs[0].GetElementByteSize();
 		const UE::NNE::FSymbolicTensorShape& SymShape = InOutputTensorDescs[0].GetShape();
 		for (int32 i = 1; i < InOutputTensorDescs.Num(); i++)
 		{
-			if (InOutputTensorDescs[i].GetElemByteSize() != ElemByteSize || SymShape != InOutputTensorDescs[i].GetShape())
+			if (InOutputTensorDescs[i].GetElementByteSize() != ElemByteSize || SymShape != InOutputTensorDescs[i].GetShape())
 			{
 				return false;
 			}
@@ -185,7 +185,7 @@ void UVertexDeltaModelInstance::CreateNNEModel()
 						TConstArrayView<UE::NNE::FTensorDesc> InputTensorDescs = ModelInstanceRDG->GetInputTensorDescs();
 						UE::NNE::FTensorShape InputTensorShape = UE::NNE::FTensorShape::MakeFromSymbolic(InputTensorDescs[0].GetShape());
 						ModelInstanceRDG->SetInputTensorShapes({ InputTensorShape });
-						check(InputTensorDescs[0].GetElemByteSize() == sizeof(float));
+						check(InputTensorDescs[0].GetElementByteSize() == sizeof(float));
 						NNEInputTensorBuffer.SetNumUninitialized(InputTensorShape.Volume());
 
 						// setup outputs
diff --git a/Engine/Plugins/Experimental/NNERuntimeORTCpu/Source/NNERuntimeORTCpu/Private/NNERuntimeORTCpu.h b/Engine/Plugins/Experimental/NNERuntimeORTCpu/Source/NNERuntimeORTCpu/Private/NNERuntimeORTCpu.h
--- a/Engine/Plugins/Experimental/NNERuntimeORTCpu/Source/NNERuntimeORTCpu/Private/NNERuntimeORTCpu.h
+++ b/Engine/Plugins/Experimental/NNERuntimeORTCpu/Source/NNERuntimeORTCpu/Private/NNERuntimeORTCpu.h
@@ -32,7 +32,6 @@ public:
 	virtual ~UNNERuntimeORTCpuImpl() {}
 		
 	virtual FString GetRuntimeName() const override { return TEXT("NNERuntimeORTCpu"); };
-	virtual bool IsPlatformSupported(const ITargetPlatform* TargetPlatform) const override { return true; };
 
 	virtual bool CanCreateModelData(FString FileType, TConstArrayView<uint8> FileData, FGuid FileId, const ITargetPlatform* TargetPlatform) const override;
 	virtual TArray<uint8> CreateModelData(FString FileType, TConstArrayView<uint8> FileData, FGuid FileId, const ITargetPlatform* TargetPlatform) override;
diff --git a/Engine/Plugins/Experimental/NNERuntimeORTCpu/Source/NNERuntimeORTCpu/Private/NNERuntimeORTCpuModel.cpp b/Engine/Plugins/Experimental/NNERuntimeORTCpu/Source/NNERuntimeORTCpu/Private/NNERuntimeORTCpuModel.cpp
--- a/Engine/Plugins/Experimental/NNERuntimeORTCpu/Source/NNERuntimeORTCpu/Private/NNERuntimeORTCpuModel.cpp
+++ b/Engine/Plugins/Experimental/NNERuntimeORTCpu/Source/NNERuntimeORTCpu/Private/NNERuntimeORTCpuModel.cpp
@@ -155,7 +155,7 @@ namespace UE::NNERuntimeORTCpu::Private
 			NNE::FSymbolicTensorShape Shape = NNE::FSymbolicTensorShape::Make(ShapeData);
 			NNE::FTensorDesc SymbolicTensorDesc = NNE::FTensorDesc::Make(FString(CurTensorNamePtr), Shape, TypeAndSize.first);
 
-			check(SymbolicTensorDesc.GetElemByteSize() == TypeAndSize.second);
+			check(SymbolicTensorDesc.GetElementByteSize() == TypeAndSize.second);
 			SymbolicTensorDescs.Emplace(SymbolicTensorDesc);
 		}
 
diff --git a/Engine/Plugins/Experimental/NNERuntimeORTGpu/Source/NNERuntimeORT/Private/NNERuntimeORT.h b/Engine/Plugins/Experimental/NNERuntimeORTGpu/Source/NNERuntimeORT/Private/NNERuntimeORT.h
--- a/Engine/Plugins/Experimental/NNERuntimeORTGpu/Source/NNERuntimeORT/Private/NNERuntimeORT.h
+++ b/Engine/Plugins/Experimental/NNERuntimeORTGpu/Source/NNERuntimeORT/Private/NNERuntimeORT.h
@@ -44,7 +44,6 @@ public:
 	void Init(ENNERuntimeORTGpuProvider Provider);
 
 	virtual FString GetRuntimeName() const override;
-	virtual bool IsPlatformSupported(const ITargetPlatform* TargetPlatform) const override { return true; };
 
 	virtual bool CanCreateModelData(FString FileType, TConstArrayView<uint8> FileData, FGuid FileId, const ITargetPlatform* TargetPlatform) const override;
 	virtual TArray<uint8> CreateModelData(FString FileType, TConstArrayView<uint8> FileData, FGuid FileId, const ITargetPlatform* TargetPlatform) override;
diff --git a/Engine/Plugins/Experimental/NNERuntimeORTGpu/Source/NNERuntimeORT/Private/NNERuntimeORTModel.cpp b/Engine/Plugins/Experimental/NNERuntimeORTGpu/Source/NNERuntimeORT/Private/NNERuntimeORTModel.cpp
--- a/Engine/Plugins/Experimental/NNERuntimeORTGpu/Source/NNERuntimeORT/Private/NNERuntimeORTModel.cpp
+++ b/Engine/Plugins/Experimental/NNERuntimeORTGpu/Source/NNERuntimeORT/Private/NNERuntimeORTModel.cpp
@@ -180,7 +180,7 @@ namespace UE::NNERuntimeORT::Private
 			NNE::FSymbolicTensorShape Shape = NNE::FSymbolicTensorShape::Make(ShapeData);
 			NNE::FTensorDesc SymbolicTensorDesc = NNE::FTensorDesc::Make(FString(TensorNames.Last()), Shape, TypeAndSize.first);
 
-			check(SymbolicTensorDesc.GetElemByteSize() == TypeAndSize.second);
+			check(SymbolicTensorDesc.GetElementByteSize() == TypeAndSize.second);
 			SymbolicTensorDescs.Emplace(SymbolicTensorDesc);
 		}
 
diff --git a/Engine/Plugins/Experimental/NNERuntimeRDG/Source/NNERuntimeRDG/Private/NNERuntimeRDGDml.h b/Engine/Plugins/Experimental/NNERuntimeRDG/Source/NNERuntimeRDG/Private/NNERuntimeRDGDml.h
--- a/Engine/Plugins/Experimental/NNERuntimeRDG/Source/NNERuntimeRDG/Private/NNERuntimeRDGDml.h
+++ b/Engine/Plugins/Experimental/NNERuntimeRDG/Source/NNERuntimeRDG/Private/NNERuntimeRDGDml.h
@@ -32,7 +32,6 @@ public:
 	bool Init(bool bRegisterOnlyOperators);
 
 	virtual FString GetRuntimeName() const override { return TEXT("NNERuntimeRDGDml"); };
-	virtual bool IsPlatformSupported(const ITargetPlatform* TargetPlatform) const override { return true; };
 
 	virtual bool CanCreateModelData(FString FileType, TConstArrayView<uint8> FileData, FGuid FileId, const ITargetPlatform* TargetPlatform) const override;
 	virtual TArray<uint8> CreateModelData(FString FileType, TConstArrayView<uint8> FileData, FGuid FileId, const ITargetPlatform* TargetPlatform) override;
diff --git a/Engine/Plugins/Experimental/NNERuntimeRDG/Source/NNERuntimeRDG/Private/NNERuntimeRDGHlsl.h b/Engine/Plugins/Experimental/NNERuntimeRDG/Source/NNERuntimeRDG/Private/NNERuntimeRDGHlsl.h
--- a/Engine/Plugins/Experimental/NNERuntimeRDG/Source/NNERuntimeRDG/Private/NNERuntimeRDGHlsl.h
+++ b/Engine/Plugins/Experimental/NNERuntimeRDG/Source/NNERuntimeRDG/Private/NNERuntimeRDGHlsl.h
@@ -27,7 +27,6 @@ public:
 	bool Init();
 
 	virtual FString GetRuntimeName() const override { return TEXT("NNERuntimeRDGHlsl"); };
-	virtual bool IsPlatformSupported(const ITargetPlatform* TargetPlatform) const override { return true; };
 
 	virtual bool CanCreateModelData(FString FileType, TConstArrayView<uint8> FileData, FGuid FileId, const ITargetPlatform* TargetPlatform) const override;
 	virtual TArray<uint8> CreateModelData(FString FileType, TConstArrayView<uint8> FileData, FGuid FileId, const ITargetPlatform* TargetPlatform) override;
diff --git a/Engine/Plugins/Experimental/NNERuntimeRDG/Source/NNERuntimeRDG/Private/NNERuntimeRDGModel.cpp b/Engine/Plugins/Experimental/NNERuntimeRDG/Source/NNERuntimeRDG/Private/NNERuntimeRDGModel.cpp
--- a/Engine/Plugins/Experimental/NNERuntimeRDG/Source/NNERuntimeRDG/Private/NNERuntimeRDGModel.cpp
+++ b/Engine/Plugins/Experimental/NNERuntimeRDG/Source/NNERuntimeRDG/Private/NNERuntimeRDGModel.cpp
@@ -188,8 +188,8 @@ int FModelInstanceRDG::SetInputTensorShapes(TConstArrayView<NNE::FTensorShape> I
 FRDGBufferDesc CreateRDGBufferDescForTensorRDG(const FTensorRDG& Tensor)
 {
 	// FIXME: CreateStructuredDesc() creates a crash on VulkanRHI
-	//FRDGBufferDesc Desc = FRDGBufferDesc::CreateStructuredDesc(Tensor.GetElemByteSize(), Tensor.GetVolume());
-	FRDGBufferDesc Desc = FRDGBufferDesc::CreateBufferDesc(Tensor.GetElemByteSize(), Tensor.GetVolume());
+	//FRDGBufferDesc Desc = FRDGBufferDesc::CreateStructuredDesc(Tensor.GetElementByteSize(), Tensor.GetVolume());
+	FRDGBufferDesc Desc = FRDGBufferDesc::CreateBufferDesc(Tensor.GetElementByteSize(), Tensor.GetVolume());
 
 	return Desc;
 }
diff --git a/Engine/Plugins/Experimental/NNERuntimeRDG/Source/NNERuntimeRDG/Private/NNERuntimeRDGModelHlsl.cpp b/Engine/Plugins/Experimental/NNERuntimeRDG/Source/NNERuntimeRDG/Private/NNERuntimeRDGModelHlsl.cpp
--- a/Engine/Plugins/Experimental/NNERuntimeRDG/Source/NNERuntimeRDG/Private/NNERuntimeRDGModelHlsl.cpp
+++ b/Engine/Plugins/Experimental/NNERuntimeRDG/Source/NNERuntimeRDG/Private/NNERuntimeRDGModelHlsl.cpp
@@ -73,7 +73,7 @@ bool FModelInstance::PrepareModelRDG(FRDGBuilder& RDGBuilder)
 		{
 			if (!PooledBuffer.IsValid())
 			{
-				FRDGBufferDesc BufferDesc = FRDGBufferDesc::CreateBufferDesc(Tensor.GetElemByteSize(), Tensor.GetVolume());
+				FRDGBufferDesc BufferDesc = FRDGBufferDesc::CreateBufferDesc(Tensor.GetElementByteSize(), Tensor.GetVolume());
 				const FRDGBufferRef TransientRDGBuffer = RDGBuilder.CreateBuffer(BufferDesc, *Tensor.GetName(), ERDGBufferFlags::None);
 				const uint8* TensorData = Tensor.GetPreparedData<uint8>().GetData();
 				PooledBuffer = RDGBuilder.ConvertToExternalBuffer(TransientRDGBuffer);
@@ -326,7 +326,7 @@ void EnqueueTensorUpload(TArray<TRefCountPtr<FRDGPooledBuffer>>& OutExternalRDGR
 					if (Tensor.HasPreparedData())
 					{
 
-						FRDGBufferDesc BufferDesc = FRDGBufferDesc::CreateBufferDesc(Tensor.GetElemByteSize(), Tensor.GetVolume());
+						FRDGBufferDesc BufferDesc = FRDGBufferDesc::CreateBufferDesc(Tensor.GetElementByteSize(), Tensor.GetVolume());
 						const FRDGBufferRef TransientRDGBuffer = RDGBuilder.CreateBuffer(BufferDesc, *Tensor.GetName(), ERDGBufferFlags::None);
 						const uint8* TensorData = Tensor.GetPreparedData<uint8>().GetData();
 
diff --git a/Engine/Source/Runtime/NNE/Internal/NNETensor.h b/Engine/Source/Runtime/NNE/Internal/NNETensor.h
--- a/Engine/Source/Runtime/NNE/Internal/NNETensor.h
+++ b/Engine/Source/Runtime/NNE/Internal/NNETensor.h
@@ -31,7 +31,7 @@ namespace UE::NNE::Internal
 			return DataType;
 		}
 
-		uint32 GetElemByteSize() const
+		uint32 GetElementByteSize() const
 		{
 			return GetTensorDataTypeSizeInBytes(DataType);
 		}
diff --git a/Engine/Source/Runtime/NNE/Public/NNERuntime.h b/Engine/Source/Runtime/NNE/Public/NNERuntime.h
--- a/Engine/Source/Runtime/NNE/Public/NNERuntime.h
+++ b/Engine/Source/Runtime/NNE/Public/NNERuntime.h
@@ -20,7 +20,7 @@ class NNE_API UNNERuntime : public UInterface
  *
  * This interface is mainly used internally by UNNEModelData to cook a model given in a file format (e.g. .onnx) into runtime specific model data.
  * The model data is then stored inside UNNEModelData which then can be used in game to create an inferable model.
- * See INNERuntimeCPU and INNERuntimeRDG on how to create a model given a UNNEModelData asset.
+ * See INNERuntimeCPU, INNERuntimeGPU and INNERuntimeRDG on how to create a model given a UNNEModelData asset.
  */
 class NNE_API INNERuntime
 {
@@ -34,16 +34,6 @@ public:
 	 * @return The name of the runtime.
 	 */
 	virtual FString GetRuntimeName() const = 0;
-	
-	/**
-	 * Check if the runtime is supported on a specific platform.
-	 *
-	 * This function is used internally in the cooking process of UNNEModelData.
-	 *
-	 * @param TargetPlatform Interface identifying the target platforms.
-	 * @return True if the platform is supported, false otherwise.
-	 */
-	virtual bool IsPlatformSupported(const ITargetPlatform* TargetPlatform) const = 0;
 
 	/**
 	 * Check if the runtime is able to create model data given some file data representing a neural network.
diff --git a/Engine/Source/Runtime/NNE/Public/NNETypes.h b/Engine/Source/Runtime/NNE/Public/NNETypes.h
--- a/Engine/Source/Runtime/NNE/Public/NNETypes.h
+++ b/Engine/Source/Runtime/NNE/Public/NNETypes.h
@@ -8,7 +8,9 @@
 #include "NNETypes.generated.h"
 
 /**
- * The enum lists all tensor data types used in NNE 
+ * The enum lists all tensor data types used in NNE.
+ *
+ * See UE::NNE::GetTensorDataTypeSizeInBytes(ENNETensorDataType InType) to get the size of each data type in bytes.
  */
 UENUM()
 enum class ENNETensorDataType : uint8
@@ -34,6 +36,13 @@ enum class ENNETensorDataType : uint8
 
 namespace UE::NNE
 {
+	/**
+	 * Return the data size in bytes for a tensor data type.
+	 *
+	 * @param InType the type of the element to consider.
+	 * @return the data size in bytes of an element.
+	 */
+	size_t NNE_API GetTensorDataTypeSizeInBytes(ENNETensorDataType InType);
 
 	/**
 	 * A symbolic tensor shape represents the shape of a tensor with potentially variable dimension.
@@ -216,14 +225,6 @@ namespace UE::NNE
 		void operator=(const FTensorShape& OtherShape);
 	};
 
-	/**
-	 * Return the data size in bytes for a tensor data type.
-	 * 
-	 * @param InType the type of the element to consider.
-	 * @return the data size in bytes of an element.
-	 */
-	size_t NNE_API GetTensorDataTypeSizeInBytes(ENNETensorDataType InType);
-
 	/**
 	 * The descriptor for a tensor as model inputs and output.
 	 *
@@ -232,6 +233,7 @@ namespace UE::NNE
 	 */
 	class NNE_API FTensorDesc
 	{
+	private:
 		FString					Name;
 		ENNETensorDataType		DataType;
 		FSymbolicTensorShape	Shape;
@@ -248,6 +250,7 @@ namespace UE::NNE
 		 * @return a tensor shape with the given dimensions.
 		 */
 		static FTensorDesc Make(const FString& Name, const FSymbolicTensorShape& Shape, ENNETensorDataType DataType);
+
 		/**
 		 * Get the name of the tensor
 		 * 
@@ -267,7 +270,7 @@ namespace UE::NNE
 		 * 
 		 * @return the size in bytes of one element of the tensor.
 		 */
-		inline uint32 GetElemByteSize() const { return GetTensorDataTypeSizeInBytes(DataType); }
+		inline uint32 GetElementByteSize() const { return GetTensorDataTypeSizeInBytes(DataType); }
 
 		/**
 		 * Get the symbolic shape of the tensor.
