[NNE] Refactor UNNEModelData and add support for memory alignment

Certain runtimes can use more efficient instruction sets when their model data is memory aligned.
This CL adds support so runtimes can express their desired memory alignment when creating model data and UNNEModelData makes sure those alignment requirements are met when deserializing data.
In addition, some refactoring has been done to support this new asset version and print proper deprecation warnings when older versions are used.

Tested by running automation tests, saving deprecated assets to upgrade to the new version and packaging and running the style transfer.

#rnx

diff --git a/Engine/Plugins/Experimental/NNERuntimeORTCpu/Source/NNERuntimeORTCpu/Private/NNERuntimeORTCpu.cpp b/Engine/Plugins/Experimental/NNERuntimeORTCpu/Source/NNERuntimeORTCpu/Private/NNERuntimeORTCpu.cpp
--- a/Engine/Plugins/Experimental/NNERuntimeORTCpu/Source/NNERuntimeORTCpu/Private/NNERuntimeORTCpu.cpp
+++ b/Engine/Plugins/Experimental/NNERuntimeORTCpu/Source/NNERuntimeORTCpu/Private/NNERuntimeORTCpu.cpp
@@ -21,7 +21,7 @@ bool UNNERuntimeORTCpuImpl::CanCreateModelData(FString FileType, TConstArrayView
 	return FileType.Compare("onnx", ESearchCase::IgnoreCase) == 0;
 }
 
-TArray<uint8> UNNERuntimeORTCpuImpl::CreateModelData(FString FileType, TConstArrayView<uint8> FileData, FGuid FileId, const ITargetPlatform* TargetPlatform)
+TArray<uint8> UNNERuntimeORTCpuImpl::CreateModelData(FString FileType, TConstArrayView<uint8> FileData, FGuid FileId, const ITargetPlatform* TargetPlatform, uint32& OutMemoryAlignment)
 {
 	if (!CanCreateModelData(FileType, FileData, FileId, TargetPlatform))
 	{
@@ -47,6 +47,8 @@ TArray<uint8> UNNERuntimeORTCpuImpl::CreateModelData(FString FileType, TConstArr
 	Writer << UNNERuntimeORTCpuImpl::GUID;
 	Writer << UNNERuntimeORTCpuImpl::Version;
 	Writer.Serialize(OutputModel.Data.GetData(), OutputModel.Data.Num());
+
+	OutMemoryAlignment = 0;
 	return Result;
 }
 
diff --git a/Engine/Plugins/Experimental/NNERuntimeORTCpu/Source/NNERuntimeORTCpu/Private/NNERuntimeORTCpu.h b/Engine/Plugins/Experimental/NNERuntimeORTCpu/Source/NNERuntimeORTCpu/Private/NNERuntimeORTCpu.h
--- a/Engine/Plugins/Experimental/NNERuntimeORTCpu/Source/NNERuntimeORTCpu/Private/NNERuntimeORTCpu.h
+++ b/Engine/Plugins/Experimental/NNERuntimeORTCpu/Source/NNERuntimeORTCpu/Private/NNERuntimeORTCpu.h
@@ -34,7 +34,7 @@ public:
 	virtual FString GetRuntimeName() const override { return TEXT("NNERuntimeORTCpu"); };
 
 	virtual bool CanCreateModelData(FString FileType, TConstArrayView<uint8> FileData, FGuid FileId, const ITargetPlatform* TargetPlatform) const override;
-	virtual TArray<uint8> CreateModelData(FString FileType, TConstArrayView<uint8> FileData, FGuid FileId, const ITargetPlatform* TargetPlatform) override;
+	virtual TArray<uint8> CreateModelData(FString FileType, TConstArrayView<uint8> FileData, FGuid FileId, const ITargetPlatform* TargetPlatform, uint32& OutMemoryAlignment) override;
 	virtual FString GetModelDataIdentifier(FString FileType, TConstArrayView<uint8> FileData, FGuid FileId, const ITargetPlatform* TargetPlatform) override;
 
 	virtual bool CanCreateModelCPU(TObjectPtr<UNNEModelData> ModelData) const override;
diff --git a/Engine/Plugins/Experimental/NNERuntimeORTGpu/Source/NNERuntimeORT/Private/NNERuntimeORT.cpp b/Engine/Plugins/Experimental/NNERuntimeORTGpu/Source/NNERuntimeORT/Private/NNERuntimeORT.cpp
--- a/Engine/Plugins/Experimental/NNERuntimeORTGpu/Source/NNERuntimeORT/Private/NNERuntimeORT.cpp
+++ b/Engine/Plugins/Experimental/NNERuntimeORTGpu/Source/NNERuntimeORT/Private/NNERuntimeORT.cpp
@@ -23,7 +23,7 @@ bool UNNERuntimeORTGpuImpl::CanCreateModelData(FString FileType, TConstArrayView
 	return FileType.Compare("onnx", ESearchCase::IgnoreCase) == 0;
 }
 
-TArray<uint8> UNNERuntimeORTGpuImpl::CreateModelData(FString FileType, TConstArrayView<uint8> FileData, FGuid FileId, const ITargetPlatform* TargetPlatform)
+TArray<uint8> UNNERuntimeORTGpuImpl::CreateModelData(FString FileType, TConstArrayView<uint8> FileData, FGuid FileId, const ITargetPlatform* TargetPlatform, uint32& OutMemoryAlignment)
 {
 	if (!CanCreateModelData(FileType, FileData, FileId, TargetPlatform))
 	{
@@ -49,6 +49,8 @@ TArray<uint8> UNNERuntimeORTGpuImpl::CreateModelData(FString FileType, TConstArr
 	Writer << UNNERuntimeORTGpuImpl::GUID;
 	Writer << UNNERuntimeORTGpuImpl::Version;
 	Writer.Serialize(OutputModel.Data.GetData(), OutputModel.Data.Num());
+
+	OutMemoryAlignment = 0;
 	return Result;
 }
 
diff --git a/Engine/Plugins/Experimental/NNERuntimeORTGpu/Source/NNERuntimeORT/Private/NNERuntimeORT.h b/Engine/Plugins/Experimental/NNERuntimeORTGpu/Source/NNERuntimeORT/Private/NNERuntimeORT.h
--- a/Engine/Plugins/Experimental/NNERuntimeORTGpu/Source/NNERuntimeORT/Private/NNERuntimeORT.h
+++ b/Engine/Plugins/Experimental/NNERuntimeORTGpu/Source/NNERuntimeORT/Private/NNERuntimeORT.h
@@ -46,7 +46,7 @@ public:
 	virtual FString GetRuntimeName() const override;
 
 	virtual bool CanCreateModelData(FString FileType, TConstArrayView<uint8> FileData, FGuid FileId, const ITargetPlatform* TargetPlatform) const override;
-	virtual TArray<uint8> CreateModelData(FString FileType, TConstArrayView<uint8> FileData, FGuid FileId, const ITargetPlatform* TargetPlatform) override;
+	virtual TArray<uint8> CreateModelData(FString FileType, TConstArrayView<uint8> FileData, FGuid FileId, const ITargetPlatform* TargetPlatform, uint32& OutMemoryAlignment) override;
 	virtual FString GetModelDataIdentifier(FString FileType, TConstArrayView<uint8> FileData, FGuid FileId, const ITargetPlatform* TargetPlatform) override;
 
 	virtual bool CanCreateModelGPU(TObjectPtr<UNNEModelData> ModelData) const override;
diff --git a/Engine/Plugins/Experimental/NNERuntimeRDG/Source/NNERuntimeRDG/Private/NNERuntimeRDGDml.cpp b/Engine/Plugins/Experimental/NNERuntimeRDG/Source/NNERuntimeRDG/Private/NNERuntimeRDGDml.cpp
--- a/Engine/Plugins/Experimental/NNERuntimeRDG/Source/NNERuntimeRDG/Private/NNERuntimeRDGDml.cpp
+++ b/Engine/Plugins/Experimental/NNERuntimeRDG/Source/NNERuntimeRDG/Private/NNERuntimeRDGDml.cpp
@@ -101,7 +101,7 @@ bool UNNERuntimeRDGDmlImpl::CanCreateModelData(FString FileType, TConstArrayView
 	return FileType.Compare("onnx", ESearchCase::IgnoreCase) == 0;
 }
 
-TArray<uint8> UNNERuntimeRDGDmlImpl::CreateModelData(FString FileType, TConstArrayView<uint8> FileData, FGuid FileId, const ITargetPlatform* TargetPlatform)
+TArray<uint8> UNNERuntimeRDGDmlImpl::CreateModelData(FString FileType, TConstArrayView<uint8> FileData, FGuid FileId, const ITargetPlatform* TargetPlatform, uint32& OutMemoryAlignment)
 {
 	if (!CanCreateModelData(FileType, FileData, FileId, TargetPlatform))
 	{
@@ -132,6 +132,8 @@ TArray<uint8> UNNERuntimeRDGDmlImpl::CreateModelData(FString FileType, TConstArr
 	Writer << Guid;
 	Writer << Version;
 	Writer.Serialize(OutputModel.Data.GetData(), OutputModel.Data.Num());
+
+	OutMemoryAlignment = 0;
 	return Result;
 };
 
diff --git a/Engine/Plugins/Experimental/NNERuntimeRDG/Source/NNERuntimeRDG/Private/NNERuntimeRDGDml.h b/Engine/Plugins/Experimental/NNERuntimeRDG/Source/NNERuntimeRDG/Private/NNERuntimeRDGDml.h
--- a/Engine/Plugins/Experimental/NNERuntimeRDG/Source/NNERuntimeRDG/Private/NNERuntimeRDGDml.h
+++ b/Engine/Plugins/Experimental/NNERuntimeRDG/Source/NNERuntimeRDG/Private/NNERuntimeRDGDml.h
@@ -34,7 +34,7 @@ public:
 	virtual FString GetRuntimeName() const override { return TEXT("NNERuntimeRDGDml"); };
 
 	virtual bool CanCreateModelData(FString FileType, TConstArrayView<uint8> FileData, FGuid FileId, const ITargetPlatform* TargetPlatform) const override;
-	virtual TArray<uint8> CreateModelData(FString FileType, TConstArrayView<uint8> FileData, FGuid FileId, const ITargetPlatform* TargetPlatform) override;
+	virtual TArray<uint8> CreateModelData(FString FileType, TConstArrayView<uint8> FileData, FGuid FileId, const ITargetPlatform* TargetPlatform, uint32& OutMemoryAlignment) override;
 	virtual FString GetModelDataIdentifier(FString FileType, TConstArrayView<uint8> FileData, FGuid FileId, const ITargetPlatform* TargetPlatform) override;
 
 	virtual bool CanCreateModelRDG(TObjectPtr<UNNEModelData> ModelData) const override;
diff --git a/Engine/Plugins/Experimental/NNERuntimeRDG/Source/NNERuntimeRDG/Private/NNERuntimeRDGHlsl.cpp b/Engine/Plugins/Experimental/NNERuntimeRDG/Source/NNERuntimeRDG/Private/NNERuntimeRDGHlsl.cpp
--- a/Engine/Plugins/Experimental/NNERuntimeRDG/Source/NNERuntimeRDG/Private/NNERuntimeRDGHlsl.cpp
+++ b/Engine/Plugins/Experimental/NNERuntimeRDG/Source/NNERuntimeRDG/Private/NNERuntimeRDGHlsl.cpp
@@ -100,7 +100,7 @@ bool UNNERuntimeRDGHlslImpl::CanCreateModelRDG(TObjectPtr<UNNEModelData> ModelDa
 	return bResult;
 };
 
-TArray<uint8> UNNERuntimeRDGHlslImpl::CreateModelData(FString FileType, TConstArrayView<uint8> FileData, FGuid FileId, const ITargetPlatform* TargetPlatform)
+TArray<uint8> UNNERuntimeRDGHlslImpl::CreateModelData(FString FileType, TConstArrayView<uint8> FileData, FGuid FileId, const ITargetPlatform* TargetPlatform, uint32& OutMemoryAlignment)
 {
 	if (!CanCreateModelData(FileType, FileData, FileId, TargetPlatform))
 	{
@@ -126,6 +126,8 @@ TArray<uint8> UNNERuntimeRDGHlslImpl::CreateModelData(FString FileType, TConstAr
 	Writer << GUID;
 	Writer << Version;
 	Writer.Serialize(OutputModel.Data.GetData(), OutputModel.Data.Num());
+
+	OutMemoryAlignment = 0;
 	return Result;
 };
 
diff --git a/Engine/Plugins/Experimental/NNERuntimeRDG/Source/NNERuntimeRDG/Private/NNERuntimeRDGHlsl.h b/Engine/Plugins/Experimental/NNERuntimeRDG/Source/NNERuntimeRDG/Private/NNERuntimeRDGHlsl.h
--- a/Engine/Plugins/Experimental/NNERuntimeRDG/Source/NNERuntimeRDG/Private/NNERuntimeRDGHlsl.h
+++ b/Engine/Plugins/Experimental/NNERuntimeRDG/Source/NNERuntimeRDG/Private/NNERuntimeRDGHlsl.h
@@ -29,7 +29,7 @@ public:
 	virtual FString GetRuntimeName() const override { return TEXT("NNERuntimeRDGHlsl"); };
 
 	virtual bool CanCreateModelData(FString FileType, TConstArrayView<uint8> FileData, FGuid FileId, const ITargetPlatform* TargetPlatform) const override;
-	virtual TArray<uint8> CreateModelData(FString FileType, TConstArrayView<uint8> FileData, FGuid FileId, const ITargetPlatform* TargetPlatform) override;
+	virtual TArray<uint8> CreateModelData(FString FileType, TConstArrayView<uint8> FileData, FGuid FileId, const ITargetPlatform* TargetPlatform, uint32& OutMemoryAlignment) override;
 	virtual FString GetModelDataIdentifier(FString FileType, TConstArrayView<uint8> FileData, FGuid FileId, const ITargetPlatform* TargetPlatform) override;
 
 	virtual bool CanCreateModelRDG(TObjectPtr<UNNEModelData> ModelData) const override;
diff --git a/Engine/Source/Editor/NNEEditor/Private/NNEEditorModelDataActions.cpp b/Engine/Source/Editor/NNEEditor/Private/NNEEditorModelDataActions.cpp
--- a/Engine/Source/Editor/NNEEditor/Private/NNEEditorModelDataActions.cpp
+++ b/Engine/Source/Editor/NNEEditor/Private/NNEEditorModelDataActions.cpp
@@ -2,7 +2,6 @@
 
 #include "NNEEditorModelDataActions.h"
 
-#include "EditorFramework/AssetImportData.h"
 #include "NNEEditorModelDataEditorToolkit.h"
 #include "NNEModelData.h"
 
@@ -43,18 +42,6 @@ namespace UE::NNEEditor::Private
 		}
 	}
 
-	void FModelDataAssetTypeActions::GetResolvedSourceFilePaths(const TArray<UObject*>& TypeAssets, TArray<FString>& OutSourceFilePaths) const
-	{
-		for (UObject* Asset : TypeAssets)
-		{
-			const UNNEModelData* ModelData = CastChecked<UNNEModelData>(Asset);
-			if (ModelData->AssetImportData)
-			{
-				ModelData->AssetImportData->ExtractFilenames(OutSourceFilePaths);
-			}
-		}
-	}
-
 } // UE::NNEEditor::Private
 
 #undef LOCTEXT_NAMESPACE
\ No newline at end of file
diff --git a/Engine/Source/Editor/NNEEditor/Private/NNEEditorModelDataActions.h b/Engine/Source/Editor/NNEEditor/Private/NNEEditorModelDataActions.h
--- a/Engine/Source/Editor/NNEEditor/Private/NNEEditorModelDataActions.h
+++ b/Engine/Source/Editor/NNEEditor/Private/NNEEditorModelDataActions.h
@@ -17,8 +17,6 @@ namespace UE::NNEEditor::Private
 		virtual FColor GetTypeColor() const override;
 		virtual UClass* GetSupportedClass() const override;
 		virtual uint32 GetCategories() override;
-		virtual bool IsImportedAsset() const override { return true; }
-		virtual void GetResolvedSourceFilePaths(const TArray<UObject*>& TypeAssets, TArray<FString>& OutSourceFilePaths) const override;
 		virtual void OpenAssetEditor(const TArray<UObject*>& InObjects, TSharedPtr<class IToolkitHost> EditWithinLevelEditor) override;
 		// End of IAssetTypeActions interface
 	};
diff --git a/Engine/Source/Editor/NNEEditor/Private/NNEEditorModelDataFactory.cpp b/Engine/Source/Editor/NNEEditor/Private/NNEEditorModelDataFactory.cpp
--- a/Engine/Source/Editor/NNEEditor/Private/NNEEditorModelDataFactory.cpp
+++ b/Engine/Source/Editor/NNEEditor/Private/NNEEditorModelDataFactory.cpp
@@ -3,10 +3,7 @@
 #include "NNEEditorModelDataFactory.h"
 
 #include "Editor.h"
-#include "Editor/EditorEngine.h"
-#include "EditorFramework/AssetImportData.h"
 #include "EngineAnalytics.h"
-#include "HAL/FileManager.h"
 #include "Kismet/GameplayStatics.h"
 #include "NNE.h"
 #include "NNEModelData.h"
@@ -35,7 +32,6 @@ UObject* UNNEModelDataFactory::FactoryCreateBinary(UClass* Class, UObject* InPar
 	TConstArrayView<uint8> BufferView = MakeArrayView(Buffer, BufferEnd-Buffer);
 	ModelData->Init(Type, BufferView);
 
-	ModelData->AssetImportData->Update(GetCurrentFilename());
 	GEditor->GetEditorSubsystem<UImportSubsystem>()->BroadcastAssetPostImport(this, ModelData);
 
 	if (FEngineAnalytics::IsAvailable())
@@ -54,83 +50,4 @@ UObject* UNNEModelDataFactory::FactoryCreateBinary(UClass* Class, UObject* InPar
 bool UNNEModelDataFactory::FactoryCanImport(const FString & Filename)
 {
 	return Filename.EndsWith(FString("onnx"));
-}
-
-bool UNNEModelDataFactory::CanReimport(UObject* Obj, TArray<FString>& OutFilenames)
-{
-	UNNEModelData* pModel = Cast<UNNEModelData>(Obj);
-	if (pModel)
-	{
-		pModel->AssetImportData->ExtractFilenames(OutFilenames);
-		return true;
-	}
-	return false;
-}
-
-void UNNEModelDataFactory::SetReimportPaths(UObject* Obj, const TArray<FString>& NewReimportPaths)
-{
-	UNNEModelData* pModel = Cast<UNNEModelData>(Obj);
-	if (pModel && ensure(NewReimportPaths.Num() == 1))
-	{
-		pModel->AssetImportData->UpdateFilenameOnly(NewReimportPaths[0]);
-	}
-}
-
-EReimportResult::Type UNNEModelDataFactory::Reimport(UObject* Obj)
-{
-	if (!Obj || !Obj->IsA(UNNEModelData::StaticClass()))
-	{
-		return EReimportResult::Failed;
-	}
-
-	UNNEModelData* pModel = Cast<UNNEModelData>(Obj);
-	const FString ResolvedSourceFilePath = pModel->AssetImportData->GetFirstFilename();
-	if (!ResolvedSourceFilePath.Len())
-	{
-		UE_LOG(LogNNE, Warning, TEXT("Cannot reimport NNE Model Data: it does not have a path stored."));
-		return EReimportResult::Failed;
-	}
-
-	// Ensure that the file provided by the path exists
-	if (IFileManager::Get().FileSize(*ResolvedSourceFilePath) == INDEX_NONE)
-	{
-		UE_LOG(LogNNE, Warning, TEXT("Cannot reimport NNE Model Data: source file cannot be found."));
-		return EReimportResult::Failed;
-	}
-
-	// Keep runtime list to be reapplied upon asset import.
-	TArray<FString, TInlineAllocator<10>> TargetRuntimes;
-	TargetRuntimes = pModel->GetTargetRuntimes();
-
-	bool OutCanceled = false;
-
-	if (ImportObject(pModel->GetClass(), pModel->GetOuter(), *pModel->GetName(), RF_Public | RF_Standalone, ResolvedSourceFilePath, nullptr, OutCanceled) != nullptr)
-	{
-		UE_LOG(LogNNE, Log, TEXT("Reimported successfully"));
-
-		pModel->SetTargetRuntimes(TargetRuntimes);
-		pModel->AssetImportData->Update(ResolvedSourceFilePath);
-
-		// Try to find the outer package so we can dirty it up
-		if (pModel->GetOuter())
-		{
-			pModel->GetOuter()->MarkPackageDirty();
-		}
-		else
-		{
-			pModel->MarkPackageDirty();
-		}
-	}
-	else if (OutCanceled)
-	{
-		UE_LOG(LogNNE, Warning, TEXT("Reimport canceled"));
-		return EReimportResult::Cancelled;
-	}
-	else
-	{
-		UE_LOG(LogNNE, Warning, TEXT("Reimport failed"));
-		return EReimportResult::Failed;
-	}
-
-	return EReimportResult::Succeeded;
 }
\ No newline at end of file
diff --git a/Engine/Source/Editor/NNEEditor/Private/NNEEditorModelDataFactory.h b/Engine/Source/Editor/NNEEditor/Private/NNEEditorModelDataFactory.h
--- a/Engine/Source/Editor/NNEEditor/Private/NNEEditorModelDataFactory.h
+++ b/Engine/Source/Editor/NNEEditor/Private/NNEEditorModelDataFactory.h
@@ -3,7 +3,6 @@
 #pragma once
 
 #include "CoreMinimal.h"
-#include "EditorReimportHandler.h"
 #include "Factories/Factory.h"
 
 #include "NNEEditorModelDataFactory.generated.h"
@@ -14,7 +13,7 @@
  * The currently supported format is .onnx. The factory is invoked when a file of this format is dragged to the editor's content browser.
  */
 UCLASS()
-class NNEEDITOR_API UNNEModelDataFactory : public UFactory, public FReimportHandler
+class NNEEDITOR_API UNNEModelDataFactory : public UFactory
 {
 	GENERATED_BODY()
 
@@ -26,10 +25,4 @@ public:
 	virtual UObject* FactoryCreateBinary(UClass* Class, UObject* InParent, FName Name, EObjectFlags Flags, UObject* Context, const TCHAR* Type, const uint8*& Buffer, const uint8* BufferEnd, FFeedbackContext* Warn) override;
 	virtual bool FactoryCanImport(const FString& Filename) override;
 	//~ End UFactory Interface
-
-	//~ Begin FReimportHandler Interface
-	virtual bool CanReimport(UObject* Obj, TArray<FString>& OutFilenames) override;
-	virtual void SetReimportPaths(UObject* Obj, const TArray<FString>& NewReimportPaths) override;
-	virtual EReimportResult::Type Reimport(UObject* Obj) override;
-	//~ End FReimportHandler Interface
 };
\ No newline at end of file
diff --git a/Engine/Source/Editor/NNEEditor/Private/NNEEditorRuntimesWidget.h b/Engine/Source/Editor/NNEEditor/Private/NNEEditorRuntimesWidget.h
--- a/Engine/Source/Editor/NNEEditor/Private/NNEEditorRuntimesWidget.h
+++ b/Engine/Source/Editor/NNEEditor/Private/NNEEditorRuntimesWidget.h
@@ -27,7 +27,6 @@ namespace UE::NNEEditor::Private
 		void OnEntryChanged(ECheckBoxState CheckType, int32 Index);
 		ECheckBoxState IsEntryChecked(int32 Index) const;
 
-		void AddCheckBox(const FString& RuntimeName, int32 Index);
 		void InitCheckBoxNames();
 
 		TArray<FString> CheckBoxNames;
diff --git a/Engine/Source/Runtime/NNE/Private/NNEModelData.cpp b/Engine/Source/Runtime/NNE/Private/NNEModelData.cpp
--- a/Engine/Source/Runtime/NNE/Private/NNEModelData.cpp
+++ b/Engine/Source/Runtime/NNE/Private/NNEModelData.cpp
@@ -2,110 +2,385 @@
 
 #include "NNEModelData.h"
 
+#include "EditorFramework/AssetImportData.h"
 #include "NNE.h"
 #include "NNEAttributeMap.h"
 #include "NNEModelOptimizerInterface.h"
 #include "NNERuntimeFormat.h"
 #include "Serialization/CustomVersion.h"
 #include "UObject/WeakInterfacePtr.h"
-#include "EditorFramework/AssetImportData.h"
 
 #if WITH_EDITOR
 #include "Containers/StringFwd.h"
-#include "DerivedDataCacheKey.h"
 #include "DerivedDataCache.h"
+#include "DerivedDataCacheKey.h"
 #include "DerivedDataRequestOwner.h"
-#include "Internationalization/TextLocalizationResource.h"
-#include "Misc/Guid.h"
-#endif
+#include "Memory/CompositeBuffer.h"
+#endif // WITH_EDITOR
 
-enum Type
+namespace UE::NNE::ModelData
 {
-	Initial = 0,
-	TargetRuntimesAndAssetImportData = 1,
-	// -----<new versions can be added before this line>-------------------------------------------------
-	// - this needs to be the last line (see note below)
-	VersionPlusOne,
-	LatestVersion = VersionPlusOne - 1
-};
+	enum Version : uint32
+	{
+		V0 = 0, // Initial
+		V1 = 1, // TargetRuntimes and AssetImportData
+		V2 = 2, // Re-arrange fields and store only ModelData in cooked assets
+		// New versions can be added above this line
+		VersionPlusOne,
+		Latest = VersionPlusOne - 1
+	};
 
-const FGuid UNNEModelData::GUID(0x9513202e, 0xeba1b279, 0xf17fe5ba, 0xab90c3f2);
-FCustomVersionRegistration NNEModelDataVersion(UNNEModelData::GUID, LatestVersion, TEXT("NNEModelDataVersion"));//Always save with the latest version
+	const FGuid GUID(0x9513202e, 0xeba1b279, 0xf17fe5ba, 0xab90c3f2);
+	FCustomVersionRegistration NNEModelDataVersion(GUID, Version::Latest, TEXT("NNEModelDataVersion"));// Always save with the latest version
+
+	const uint32 DDCAssetVersion = 0; // Increase this value to force rebuilding cache entries
+
+	FString GetRuntimesAsString(TArrayView<const FString> Runtimes)
+	{
+		if (Runtimes.Num() == 0)
+		{
+			return TEXT("All");
+		}
+
+		FString RuntimesAsOneString;
+		bool bIsFirstRuntime = true;
+
+		for (const FString& Runtime : Runtimes)
+		{
+			if (!bIsFirstRuntime)
+			{
+				RuntimesAsOneString += TEXT(", ");
+			}
+			RuntimesAsOneString += Runtime;
+			bIsFirstRuntime = false;
+		}
+		return RuntimesAsOneString;
+	}
 
 #if WITH_EDITOR
 
-inline FString GetDDCRequestId(const FString& FileId, const FString& RuntimeName, const FString& ModelDataIdentifier)
-{
-	//RuntimeName and FileId are embedded to the id to ensure no potential collision between the runtime/assets
-	return RuntimeName + "-" + FileId + "-" + ModelDataIdentifier;
-}
+	inline FString GetDDCRequestId(const FString& FileId, const FString& RuntimeName, const FString& ModelDataIdentifier)
+	{
+		// RuntimeName and FileId are embedded to the id to ensure no potential collision between the runtime/assets
+		return RuntimeName + "-" + FileId + "-DDCv" + FString::FromInt(DDCAssetVersion) + "-" + ModelDataIdentifier;
+	}
 
-inline UE::DerivedData::FCacheKey CreateCacheKey(const FString& FileId, const FString& RequestId)
-{
-	return { UE::DerivedData::FCacheBucket(FWideStringView(*FileId)), FIoHash::HashBuffer(MakeMemoryView(FTCHARToUTF8(RequestId))) };
-}
+	inline UE::DerivedData::FCacheKey CreateCacheKey(const FString& FileId, const FString& RequestId)
+	{
+		return { UE::DerivedData::FCacheBucket(FWideStringView(*FileId)), FIoHash::HashBuffer(MakeMemoryView(FTCHARToUTF8(RequestId))) };
+	}
 
-inline FSharedBuffer GetFromDDC(const FGuid& FileId, const FString& RuntimeName, const FString& ModelDataIdentifier)
-{
-	FString FileIdStr = FileId.ToString(EGuidFormats::Digits);
-	FString RequestId = GetDDCRequestId(FileIdStr, RuntimeName, ModelDataIdentifier);
-
-	UE::DerivedData::FCacheGetValueRequest GetRequest;
-	GetRequest.Name = FString("Get-") + RequestId;
-	GetRequest.Key = CreateCacheKey(FileIdStr, RequestId);
-	FSharedBuffer RawDerivedData;
-	UE::DerivedData::FRequestOwner BlockingGetOwner(UE::DerivedData::EPriority::Blocking);
-	UE::DerivedData::GetCache().GetValue({ GetRequest }, BlockingGetOwner, [&RawDerivedData](UE::DerivedData::FCacheGetValueResponse&& Response)
+	inline void PutIntoDDC(const FGuid& FileId, const FString& RuntimeName, const FString& ModelDataIdentifier, FSharedBuffer& Data, uint32 MemoryAlignment)
+	{
+		FString FileIdStr = FileId.ToString(EGuidFormats::Digits);
+		FString RequestId = GetDDCRequestId(FileIdStr, RuntimeName, ModelDataIdentifier);
+
+		TArray<UE::DerivedData::FCachePutValueRequest> Requests;
+		Requests.SetNum(1);
+
+		Requests[0].Name = FString("Put-") + RequestId;
+		Requests[0].Key = CreateCacheKey(FileIdStr, RequestId);
+		Requests[0].Value = UE::DerivedData::FValue::Compress(FCompositeBuffer(MakeSharedBufferFromArray(TArray<uint32>({ MemoryAlignment })), Data));
+
+		UE::DerivedData::FRequestOwner BlockingPutOwner(UE::DerivedData::EPriority::Blocking);
+		UE::DerivedData::GetCache().PutValue(Requests, BlockingPutOwner);
+		BlockingPutOwner.Wait();
+	}
+
+	inline FSharedBuffer GetFromDDC(const FGuid& FileId, const FString& RuntimeName, const FString& ModelDataIdentifier, uint32& OutMemoryAlignment)
+	{
+		FString FileIdStr = FileId.ToString(EGuidFormats::Digits);
+		FString RequestId = GetDDCRequestId(FileIdStr, RuntimeName, ModelDataIdentifier);
+
+		TArray<UE::DerivedData::FCacheGetValueRequest> Requests;
+		Requests.SetNum(1);
+
+		Requests[0].Name = FString("Get-") + RequestId;
+		Requests[0].Key = CreateCacheKey(FileIdStr, RequestId);
+
+		FSharedBuffer Result;
+		UE::DerivedData::FRequestOwner BlockingGetOwner(UE::DerivedData::EPriority::Blocking);
+		UE::DerivedData::GetCache().GetValue(Requests, BlockingGetOwner, [&Result, &OutMemoryAlignment](UE::DerivedData::FCacheGetValueResponse&& Response)
 		{
-			RawDerivedData = Response.Value.GetData().Decompress();
+			if (Response.Value.HasData() && Response.Value.GetRawSize() > sizeof(uint32))
+			{
+				FCompressedBufferReader Reader(Response.Value.GetData());
+				uint32 MemoryAlignment = ((uint32*)Reader.Decompress(0, sizeof(uint32)).GetData())[0];
+				uint64 DataSize = Response.Value.GetRawSize() - sizeof(uint32);
+
+				void* Data = FMemory::Malloc(DataSize, MemoryAlignment);
+				if (Reader.TryDecompressTo(FMutableMemoryView(Data, DataSize), sizeof(uint32)))
+				{
+					OutMemoryAlignment = MemoryAlignment;
+					Result = FSharedBuffer::TakeOwnership(Data, DataSize, FMemory::Free);
+				}
+				else
+				{
+					FMemory::Free(Data);
+				}
+			}
 		});
-	BlockingGetOwner.Wait();
-	return RawDerivedData;
-}
+		BlockingGetOwner.Wait();
+		return Result;
+	}
+
+#endif // WITH_EDITOR
+
+	inline FSharedBuffer CreateModelData(const FString& RuntimeName, FString FileType, const TArray<uint8>& FileData, FGuid FileId, const ITargetPlatform* TargetPlatform, uint32& OutMemoryAlignment)
+	{
+		TWeakInterfacePtr<INNERuntime> NNERuntime = UE::NNE::GetRuntime<INNERuntime>(RuntimeName);
+		if (NNERuntime.IsValid())
+		{
+			if (NNERuntime->CanCreateModelData(FileType, FileData, FileId, TargetPlatform))
+			{
+				uint32 TmpOutMemoryAlignment = 0;
+				TArray<uint8> ModelData = NNERuntime->CreateModelData(FileType, FileData, FileId, TargetPlatform, TmpOutMemoryAlignment);
+
+				// Make sure a runtime does fulfill its own alignment requirement
+				checkf(TmpOutMemoryAlignment <= 1 || (((uintptr_t)(const void *)(ModelData.GetData())) % TmpOutMemoryAlignment == 0), TEXT("Runtimes must return ModelData that is aligned with OutMemoryAlignment!"))
+
+				if (ModelData.Num() > 0)
+				{
+					OutMemoryAlignment = TmpOutMemoryAlignment;
+					return MakeSharedBufferFromArray(MoveTemp(ModelData));
+				}
+			}
+		}
+		else
+		{
+			UE_LOG(LogNNE, Error, TEXT("UNNEModelData: No runtime '%s' found. Valid runtimes are: "), *RuntimeName);
+			TArrayView<TWeakInterfacePtr<INNERuntime>> Runtimes = UE::NNE::GetAllRuntimes();
+			for (int32 i = 0; i < Runtimes.Num(); i++)
+			{
+				UE_LOG(LogNNE, Error, TEXT("- %s"), *Runtimes[i]->GetRuntimeName());
+			}
+		}
+		return {};
+	}
 
-inline void PutIntoDDC(const FGuid& FileId, const FString& RuntimeName, const FString& ModelDataIdentifier, FSharedBuffer& Data)
+} // UE::NNE::ModelData
+
+void UNNEModelData::GetAssetRegistryTags(TArray<FAssetRegistryTag>& OutTags) const
 {
-	FString FileIdStr = FileId.ToString(EGuidFormats::Digits);
-	FString RequestId = GetDDCRequestId(FileIdStr, RuntimeName, ModelDataIdentifier);
-
-	UE::DerivedData::FCachePutValueRequest PutRequest;
-	PutRequest.Name = FString("Put-") + RequestId;
-	PutRequest.Key = CreateCacheKey(FileIdStr, RequestId);
-	PutRequest.Value = UE::DerivedData::FValue::Compress(Data);
-	UE::DerivedData::FRequestOwner BlockingPutOwner(UE::DerivedData::EPriority::Blocking);
-	UE::DerivedData::GetCache().PutValue({ PutRequest }, BlockingPutOwner);
-	BlockingPutOwner.Wait();
+	OutTags.Add(FAssetRegistryTag("TargetRuntimes", UE::NNE::ModelData::GetRuntimesAsString(GetTargetRuntimes()), FAssetRegistryTag::TT_Alphabetical));
+	Super::GetAssetRegistryTags(OutTags);
 }
 
-#endif
-
-inline TArray<uint8> CreateRuntimeDataBlob(const FString& RuntimeName, FString FileType, const TArray<uint8>& FileData, FGuid FileId, const ITargetPlatform* TargetPlatform)
+void UNNEModelData::Serialize(FArchive& Ar)
 {
-	TWeakInterfacePtr<INNERuntime> NNERuntime = UE::NNE::GetRuntime<INNERuntime>(RuntimeName);
-	if (NNERuntime.IsValid())
+	// Store the asset version (no effect in load)
+	Ar.UsingCustomVersion(UE::NNE::ModelData::GUID);
+
+	if (Ar.IsSaving())
 	{
-		return NNERuntime->CreateModelData(FileType, FileData, FileId, TargetPlatform);
+		bool bWriteModelData = true;
+		if (Ar.IsCooking())
+		{
+			// Optimize storage: FileData is not required anymore because we have the model and can cook it for every runtime
+			TArray<FString> TmpTargetRuntimes;
+			Ar << TmpTargetRuntimes;
+			FString TmpFileType;
+			Ar << TmpFileType;
+			TArray<uint8> TmpFileData;
+			Ar << TmpFileData;
+			Ar << FileId;
+
+			// Cooking must recreate all model data but only if file data is still available
+			if (FileData.Num() > 0)
+			{
+				ModelData.Reset();
+
+				// No target runtime means all currently registered ones
+				TArray<FString, TInlineAllocator<10>> CookRuntimeNames;
+				if (GetTargetRuntimes().IsEmpty())
+				{
+					for (const TWeakInterfacePtr<INNERuntime>& Runtime : UE::NNE::GetAllRuntimes())
+					{
+						CookRuntimeNames.Add(Runtime->GetRuntimeName());
+					}
+				}
+				else
+				{
+					CookRuntimeNames.Append(GetTargetRuntimes());
+				}
+
+				for (const FString& RuntimeName : CookRuntimeNames)
+				{
+					uint32 MemoryAlignment = 0;
+					FSharedBuffer CreatedData = UE::NNE::ModelData::CreateModelData(RuntimeName, FileType, FileData, FileId, Ar.GetArchiveState().CookingTarget(), MemoryAlignment);
+					if (CreatedData.GetSize() > 0)
+					{
+						ModelData.Add(RuntimeName, MakeTuple(CreatedData, MemoryAlignment));
+#if WITH_EDITOR
+						TWeakInterfacePtr<INNERuntime> NNERuntime = UE::NNE::GetRuntime<INNERuntime>(RuntimeName);
+						if (NNERuntime.IsValid())
+						{
+							FString ModelDataIdentifier = NNERuntime->GetModelDataIdentifier(FileType, FileData, FileId, Ar.GetArchiveState().CookingTarget());
+							if (ModelDataIdentifier.Len() > 0)
+							{
+								UE::NNE::ModelData::PutIntoDDC(FileId, RuntimeName, ModelDataIdentifier, CreatedData, MemoryAlignment);
+							}
+							else
+							{
+								UE_LOG(LogNNE, Warning, TEXT("UNNEModelData: Runtime '%s' returned an empty string as a ModelDataIdentifier while cooking. GetModelDataIdentifier should always return a valid identifier."), *RuntimeName);
+							}
+						}
+						else
+						{
+							UE_LOG(LogNNE, Warning, TEXT("UNNEModelData: Runtime '%s' is among the cooked runtimes but instance is invalid."), *RuntimeName);
+						}
+#endif //WITH_EDITOR
+					}
+				}
+			}
+		}
+		else
+		{
+			// Only cooked assets optimize storage
+			Ar << TargetRuntimes;
+			Ar << FileType;
+			Ar << FileData;
+			Ar << FileId;
+
+#if WITH_EDITOR
+			// In editor (when not cooking), no model data is stored as model data can always be recreated and unnecessary data in subversion control should be avoided
+			bWriteModelData = false;
+#endif //WITH_EDITOR
+		}
+
+		if (bWriteModelData)
+		{
+			TArray<FString> RuntimeNames;
+			ModelData.GetKeys(RuntimeNames);
+			int32 NumItems = RuntimeNames.Num();
+
+			Ar << NumItems;
+			for (int32 i = 0; i < NumItems; i++)
+			{
+				Ar << RuntimeNames[i];
+
+				uint32 MemoryAlignment = ModelData[RuntimeNames[i]].Get<1>();
+				Ar << MemoryAlignment;
+
+				uint64 DataSize = ModelData[RuntimeNames[i]].Get<0>().GetSize();
+				Ar << DataSize;
+
+				Ar.Serialize((void*)ModelData[RuntimeNames[i]].Get<0>().GetData(), DataSize);
+			}
+		}
+		else
+		{
+			int32 NumItems = 0;
+			Ar << NumItems;
+		}
 	}
 	else
 	{
-		UE_LOG(LogNNE, Error, TEXT("UNNEModelData: No runtime '%s' found. Valid runtimes are: "), *RuntimeName);
-		TArrayView<TWeakInterfacePtr<INNERuntime>> Runtimes = UE::NNE::GetAllRuntimes();
-		for (int i = 0; i < Runtimes.Num(); i++)
+		// Read the archive
+
+		TObjectPtr<class UAssetImportData> AssetImportData;
+		int32 NumItems;
+		FString Name;
+		uint32 MemoryAlignment;
+		uint64 DataSize;
+		TArray<uint8> Data;
+		void* RawData;
+		int32 Index;
+
+		switch (Ar.CustomVer(UE::NNE::ModelData::GUID))
 		{
-			UE_LOG(LogNNE, Error, TEXT("- %s"), *Runtimes[i]->GetRuntimeName());
+		case UE::NNE::ModelData::Version::V0:
+			TargetRuntimes.Empty();
+			Ar << FileType;
+			Ar << FileData;
+			Ar << FileId;
+			Ar << NumItems;
+			for (Index = 0; Index < NumItems; Index++)
+			{
+				Ar << Name;
+				Ar << Data;
+				ModelData.Add(Name, MakeTuple(MakeSharedBufferFromArray(MoveTemp(Data)), 0));
+			}
+			UE_LOG(LogNNE, Warning, TEXT("[DEPRECATION] UNNEModelData: The asset %s (v0) is deprecated. Please right-click the asset and select 'Save' to update it to the latest version."), *this->GetName());
+			break;
+
+		case UE::NNE::ModelData::Version::V1:
+			TargetRuntimes.Empty();
+			if (!Ar.IsLoadingFromCookedPackage())
+			{
+				Ar << TargetRuntimes;
+				Ar << AssetImportData;
+			}
+			Ar << FileType;
+			Ar << FileData;
+			Ar << FileId;
+			Ar << NumItems;
+			for (Index = 0; Index < NumItems; Index++)
+			{
+				Ar << Name;
+				Ar << Data;
+				ModelData.Add(Name, MakeTuple(MakeSharedBufferFromArray(MoveTemp(Data)), 0));
+			}
+			UE_LOG(LogNNE, Warning, TEXT("[DEPRECATION] UNNEModelData: The asset %s (v1) is deprecated. Please right-click the asset and select 'Save' to update it to the latest version."), *this->GetName());
+			break;
+
+		case UE::NNE::ModelData::Version::V2:
+			Ar << TargetRuntimes;
+			Ar << FileType;
+			Ar << FileData;
+			Ar << FileId;
+			Ar << NumItems;
+			for (Index = 0; Index < NumItems; Index++)
+			{
+				Ar << Name;
+				Ar << MemoryAlignment;
+				Ar << DataSize;
+				RawData = FMemory::Malloc(DataSize, MemoryAlignment);
+				Ar.Serialize(RawData, DataSize);
+				ModelData.Add(Name, MakeTuple(FSharedBuffer::TakeOwnership(RawData, DataSize, FMemory::Free), MemoryAlignment));
+			}
+			break;
+
+		default:
+			UE_LOG(LogNNE, Error, TEXT("UNNEModelData: Unknown asset version %d: Deserialisation failed, please reimport the original model."), Ar.CustomVer(UE::NNE::ModelData::GUID));
+			break;
 		}
-		return {};
 	}
 }
 
 void UNNEModelData::Init(const FString& Type, TConstArrayView<uint8> Buffer)
 {
+	TargetRuntimes.Empty();
 	FileType = Type;
 	FileData = Buffer;
 	FPlatformMisc::CreateGuid(FileId);
 	ModelData.Empty();
 }
 
+TArrayView<const FString> UNNEModelData::GetTargetRuntimes() const 
+{ 
+	return TargetRuntimes;
+}
+
+void UNNEModelData::SetTargetRuntimes(TArrayView<const FString> RuntimeNames)
+{
+	TargetRuntimes = RuntimeNames;
+
+	if (RuntimeNames.Num() > 0)
+	{
+		TArray<FString, TInlineAllocator<10>> CookedRuntimes;
+		ModelData.GetKeys(CookedRuntimes);
+		for (const FString& Runtime : CookedRuntimes)
+		{
+			if (!TargetRuntimes.Contains(Runtime))
+			{
+				ModelData.Remove(Runtime);
+			}
+		}
+		ModelData.Compact();
+	}
+}
+
 FString UNNEModelData::GetFileType()
 {
 	return FileType;
@@ -116,6 +391,12 @@ TConstArrayView<uint8> UNNEModelData::GetFileData()
 	return FileData;
 }
 
+void UNNEModelData::ClearFileDataAndFileType()
+{
+	FileType = "";
+	FileData.Empty();
+}
+
 FGuid UNNEModelData::GetFileId()
 {
 	return FileId;
@@ -123,7 +404,6 @@ FGuid UNNEModelData::GetFileId()
 
 TSharedPtr<UE::NNE::FSharedModelData> UNNEModelData::GetModelData(const FString& RuntimeName)
 {
-#if WITH_EDITORONLY_DATA
 	// Check model data is supporting the requested target runtime
 	TArrayView<const FString> TargetRuntimesNames = GetTargetRuntimes();
 	if (!TargetRuntimesNames.IsEmpty() && !TargetRuntimesNames.Contains(RuntimeName))
@@ -135,15 +415,21 @@ TSharedPtr<UE::NNE::FSharedModelData> UNNEModelData::GetModelData(const FString&
 		}
 		return TSharedPtr<UE::NNE::FSharedModelData>();
 	}
-#endif //WITH_EDITORONLY_DATA
 
 	// Check if we have a local cache hit
-	FSharedBuffer* LocalData = ModelData.Find(RuntimeName);
+	TTuple<FSharedBuffer, uint32>* LocalData = ModelData.Find(RuntimeName);
 	if (LocalData)
 	{
-		return MakeShared<UE::NNE::FSharedModelData>(*LocalData);
+		return MakeShared<UE::NNE::FSharedModelData>(LocalData->Get<0>());
+	}
+
+	// After this point FileData is required to either get the cache id or recreate it from scratch
+	if (FileData.Num() < 1)
+	{
+		UE_LOG(LogNNE, Error, TEXT("UNNEModelData: Cannot create model data from empty file data."));
+		return TSharedPtr<UE::NNE::FSharedModelData>();
 	}
-	
+
 #if WITH_EDITOR
 	TWeakInterfacePtr<INNERuntime> NNERuntime = UE::NNE::GetRuntime<INNERuntime>(RuntimeName);
 	if (!NNERuntime.IsValid())
@@ -158,207 +444,38 @@ TSharedPtr<UE::NNE::FSharedModelData> UNNEModelData::GetModelData(const FString&
 		UE_LOG(LogNNE, Error, TEXT("UNNEModelData: Runtime '%s' returned an empty string as a ModelDataIdentifier. GetModelDataIdentifier should always return a valid identifier."), *RuntimeName);
 		return TSharedPtr<UE::NNE::FSharedModelData>();
 	}
-	
+
 	// Check if we have a DDC cache hit
-	FSharedBuffer RemoteData = GetFromDDC(FileId, RuntimeName, ModelDataIdentifier);
+	uint32 RemoteMemoryAlignment = 0;
+	FSharedBuffer RemoteData = UE::NNE::ModelData::GetFromDDC(FileId, RuntimeName, ModelDataIdentifier, RemoteMemoryAlignment);
 	if (RemoteData.GetSize() > 0)
 	{
-		ModelData.Add(RuntimeName, RemoteData);
-		
+		ModelData.Add(RuntimeName, MakeTuple(RemoteData, RemoteMemoryAlignment));
+
 		return MakeShared<UE::NNE::FSharedModelData>(RemoteData);
 	}
 #endif //WITH_EDITOR
 
 	// Try to create the model
-	TArray<uint8> RuntimeDataBlob = CreateRuntimeDataBlob(RuntimeName, FileType, FileData, FileId, nullptr);
-	FSharedBuffer CreatedData = MakeSharedBufferFromArray(MoveTemp(RuntimeDataBlob));
+	uint32 CreatedMemoryAlignment = 0;
+	FSharedBuffer CreatedData = UE::NNE::ModelData::CreateModelData(RuntimeName, FileType, FileData, FileId, nullptr, CreatedMemoryAlignment);
 	if (CreatedData.GetSize() < 1)
 	{
-		
 		return TSharedPtr<UE::NNE::FSharedModelData>();
 	}
 
 	// Cache the model
-	ModelData.Add(RuntimeName, CreatedData);
+	ModelData.Add(RuntimeName, MakeTuple(CreatedData, CreatedMemoryAlignment));
 
 #if WITH_EDITOR
 	// And put it into DDC
-	PutIntoDDC(FileId, RuntimeName, ModelDataIdentifier, CreatedData);
+	UE::NNE::ModelData::PutIntoDDC(FileId, RuntimeName, ModelDataIdentifier, CreatedData, CreatedMemoryAlignment);
 #endif //WITH_EDITOR
-	
-	return MakeShared<UE::NNE::FSharedModelData>(CreatedData);
-}
 
-void UNNEModelData::Serialize(FArchive& Ar)
-{
-	// Store the asset version (no effect in load)
-	Ar.UsingCustomVersion(UNNEModelData::GUID);
-
-#if WITH_EDITORONLY_DATA
-	// Recreate each model data when cooking
-	if (Ar.IsCooking() && Ar.IsSaving())
-	{
-		ModelData.Reset();
-
-		TArray<FString, TInlineAllocator<10>> CookedRuntimeNames;
-		CookedRuntimeNames.Append(GetTargetRuntimes());
-
-		//No target runtime means all currently registered ones.
-		if (GetTargetRuntimes().IsEmpty())
-		{
-			for (const TWeakInterfacePtr<INNERuntime>& Runtime : UE::NNE::GetAllRuntimes())
-			{
-				CookedRuntimeNames.Add(Runtime->GetRuntimeName());
-			}
-		}
-
-		for (const FString& RuntimeName : CookedRuntimeNames)
-		{
-			TArray RuntimeDataBlob = CreateRuntimeDataBlob(RuntimeName, FileType, FileData, FileId, Ar.GetArchiveState().CookingTarget());
-			if (RuntimeDataBlob.Num() > 0)
-			{
-				FSharedBuffer CreatedData = MakeSharedBufferFromArray(MoveTemp(RuntimeDataBlob));
-				ModelData.Add(RuntimeName, CreatedData);
-#if WITH_EDITOR
-				TWeakInterfacePtr<INNERuntime> NNERuntime = UE::NNE::GetRuntime<INNERuntime>(RuntimeName);
-				if (NNERuntime.IsValid())
-				{
-					FString ModelDataIdentifier = NNERuntime->GetModelDataIdentifier(FileType, FileData, FileId, Ar.GetArchiveState().CookingTarget());
-					if (ModelDataIdentifier.Len() > 0)
-					{
-						PutIntoDDC(FileId, RuntimeName, ModelDataIdentifier, CreatedData);
-					}
-					else
-					{
-						UE_LOG(LogNNE, Warning, TEXT("UNNEModelData: Runtime '%s' returned an empty string as a ModelDataIdentifier while cooking. GetModelDataIdentifier should always return a valid identifier."), *RuntimeName);
-					}
-				}
-				else
-				{
-					UE_LOG(LogNNE, Warning, TEXT("UNNEModelData: Runtime '%s' is among the cooked runtimes but instance is invalid."), *RuntimeName);
-				}
-#endif //WITH_EDITOR
-			}
-		}
-
-		// Dummy data for fields not required in the game
-		TArray<uint8> EmptyData;
-		TArray<FString> RuntimeNames;
-		ModelData.GetKeys(RuntimeNames);
-		int32 NumItems = RuntimeNames.Num();
-
-		Ar << FileType;
-		Ar << EmptyData;
-		Ar << FileId;
-		Ar << NumItems;
-
-		for (int i = 0; i < NumItems; i++)
-		{
-			Ar << RuntimeNames[i];
-			FSharedBuffer Data = ModelData[RuntimeNames[i]];
-			TArray<uint8> DataToSerialize(static_cast<const uint8*>(Data.GetData()), Data.GetSize());
-			Ar << DataToSerialize;
-		}
-	}
-	else
-#endif //WITH_EDITORONLY_DATA
-	{
-		int32 NumItems = 0;
-
-#if WITH_EDITORONLY_DATA
-		if (Ar.CustomVer(UNNEModelData::GUID) >= TargetRuntimesAndAssetImportData)
-		{
-			Ar << TargetRuntimes;
-			Ar << AssetImportData;
-		}
-		else 
-		{
-			// AssetImportData should always be valid
-			AssetImportData = NewObject<UAssetImportData>(this, TEXT("AssetImportData"));
-		}
-#endif //WITH_EDITORONLY_DATA
-
-		Ar << FileType;
-		Ar << FileData;
-		Ar << FileId;
-		Ar << NumItems;
-
-		if (Ar.IsLoading())
-		{
-			for (int i = 0; i < NumItems; i++)
-			{
-				FString Name;
-				Ar << Name;
-				TArray<uint8> Data;
-				Ar << Data;
-				FSharedBuffer SharedData = MakeSharedBufferFromArray(MoveTemp(Data));
-				ModelData.Add(Name, SharedData);
-			}
-		}
-	}
-}
-
-#if WITH_EDITORONLY_DATA
-namespace UE::NNE::ModelDataHelpers
-{
-	FString GetRuntimesAsString(TArrayView<const FString> Runtimes)
-	{
-		if (Runtimes.Num() == 0)
-		{
-			return TEXT("All");
-		}
-
-		FString RuntimesAsOneString;
-		bool bIsFirstRuntime = true;
-
-		for (const FString& Runtime : Runtimes)
-		{
-			if (!bIsFirstRuntime)
-			{
-				RuntimesAsOneString += TEXT(", ");
-			}
-			RuntimesAsOneString += Runtime;
-			bIsFirstRuntime = false;
-		}
-		return RuntimesAsOneString;
-	}
-} // UE::NNE::ModelDataHelpers
-
-void UNNEModelData::PostInitProperties()
-{
-	if (!HasAnyFlags(RF_ClassDefaultObject))
-	{
-		AssetImportData = NewObject<UAssetImportData>(this, TEXT("AssetImportData"));
-	}
-	Super::PostInitProperties();
-}
-
-void UNNEModelData::GetAssetRegistryTags(TArray<FAssetRegistryTag>& OutTags) const
-{
-	if (AssetImportData)
-	{
-		OutTags.Add(FAssetRegistryTag(SourceFileTagName(), AssetImportData->GetSourceData().ToJson(), FAssetRegistryTag::TT_Hidden));
-	}
-
-	OutTags.Add(FAssetRegistryTag("TargetRuntimes", UE::NNE::ModelDataHelpers::GetRuntimesAsString(GetTargetRuntimes()), FAssetRegistryTag::TT_Alphabetical));
-
-	Super::GetAssetRegistryTags(OutTags);
+	return MakeShared<UE::NNE::FSharedModelData>(CreatedData);
 }
 
-void UNNEModelData::SetTargetRuntimes(TArrayView<const FString> RuntimeNames)
+void UNNEModelData::ClearModelData()
 {
-	TargetRuntimes = RuntimeNames;
-
-	TArray<FString, TInlineAllocator<10>> CookedRuntimes;
-	ModelData.GetKeys(CookedRuntimes);
-	for (const FString& Runtime : CookedRuntimes)
-	{
-		if (!TargetRuntimes.Contains(Runtime))
-		{
-			ModelData.Remove(Runtime);
-		}
-	}
-	ModelData.Compact();
+	ModelData.Empty();
 }
-
-#endif //WITH_EDITORONLY_DATA
\ No newline at end of file
diff --git a/Engine/Source/Runtime/NNE/Public/NNEModelData.h b/Engine/Source/Runtime/NNE/Public/NNEModelData.h
--- a/Engine/Source/Runtime/NNE/Public/NNEModelData.h
+++ b/Engine/Source/Runtime/NNE/Public/NNEModelData.h
@@ -4,6 +4,7 @@
 
 #include "CoreMinimal.h"
 #include "Memory/SharedBuffer.h"
+#include "Misc/Guid.h"
 #include "Templates/SharedPointer.h"
 
 #include "NNEModelData.generated.h"
@@ -23,6 +24,11 @@ namespace UE::NNE
 		FSharedModelData(FSharedBuffer InData) : Data(InData) {}
 		FSharedModelData() {}
 
+		/**
+		 * Get a const array view on the shared data which is guaranteed to remain valid as long as this objects exists.
+		 *
+		 * @return A const array view of the shared data.
+		 */
 		TConstArrayView<uint8> GetView() const
 		{
 			return MakeArrayView(static_cast<const uint8*>(Data.GetData()), Data.GetSize());
@@ -44,6 +50,11 @@ class NNE_API UNNEModelData : public UObject
 
 public:
 
+	// UObject interface
+	virtual void GetAssetRegistryTags(TArray<FAssetRegistryTag>& OutTags) const override;
+	virtual void Serialize(FArchive& Ar) override;
+	// End of UObject interface
+
 	/**
 	 * Initialize the model data with a copy of the data inside Buffer.
 	 *
@@ -54,6 +65,20 @@ public:
 	 */
 	void Init(const FString& Type, TConstArrayView<uint8> Buffer);
 
+	/**
+	 * Get the target runtimes this model data will be cooked for. An empty list means all runtimes.
+	 *
+	 * @return The target runtimes names.
+	 */
+	TArrayView<const FString> GetTargetRuntimes() const;
+
+	/**
+	 * Set the target runtimes this model data will be cooked for. An empty list means all runtimes.
+	 *
+	 * @param RuntimeNames The target runtimes names.
+	 */
+	void SetTargetRuntimes(TArrayView<const FString> RuntimeNames);
+
 	/**
 	 * Get the type of data inside FileData.
 	 *
@@ -72,6 +97,13 @@ public:
 	 */
 	TConstArrayView<uint8> GetFileData();
 
+	/**
+	 * Clears the FileData and the FileType.
+	 *
+	 * Caution, if the FileData is cleared, no more models can be created on runtimes that do not already have ModelData inside this asset.
+	 */
+	void ClearFileDataAndFileType();
+
 	/**
 	 * Get the FGuid identifying the FileData.
 	 *
@@ -84,7 +116,7 @@ public:
 	/**
 	 * Get the cached (editor) or cooked (game) optimized model data for a given runtime.
 	 *
-	 * This function is used by runtimes when creating a model. In editor, the function will create the optimized model data with the passed runtime in case it has not been cached in the DCC yet. In game, the cooked data is accessed.
+	 * This function is used by runtimes when creating a model. In editor, the function will create the optimized model data with the passed runtime in case it has not been cached in the DCC yet. In game, the cooked data is accessed. The returned model data is aligned in memory as requested by the runtime.
 	 *
 	 * @param RuntimeName The name of the runtime for which the data should be returned.
 	 * @return The optimized and runtime specific model data or an invalid TSharedPtr in case of failure.
@@ -92,50 +124,17 @@ public:
 	TSharedPtr<UE::NNE::FSharedModelData> GetModelData(const FString& RuntimeName);
 
 	/**
-	 * Implements custom serialization of this asset.
-	 * @param Ar The archive to serialize from/to.
-	 */
-	virtual void Serialize(FArchive& Ar) override;
-
-	/**
-	 * A Guid used for asset versioning.
-	 */
-	const static FGuid GUID;
-
-#if WITH_EDITORONLY_DATA
-	// UObject interface
-	virtual void PostInitProperties() override;
-	virtual void GetAssetRegistryTags(TArray<FAssetRegistryTag>& OutTags) const override;
-	// End of UObject interface
-
-	/**
-	 * Get the target runtimes this model data will be cooked for. An empty list mean all runtimes.
-	 *
-	 * @return The target runtimes names.
-	 */
-	TArrayView<const FString> GetTargetRuntimes() const { return TargetRuntimes; }
-
-	/**
-	 * Set the target runtimes this model data will be cooked for. An empty list mean all runtimes.
+	 * Clears the ModelData.
 	 *
-	 * @param RuntimeNames The target runtimes names.
+	 * Caution, if the ModelData is cleared, only runtimes that support cooking on the current platform can create new models from this asset.
 	 */
-	void SetTargetRuntimes(TArrayView<const FString> RuntimeNames);
-
-	/**
-	 * Importing data used for this asset.
-	 */
-	UPROPERTY(VisibleAnywhere, Instanced, Category = ImportSettings)
-	TObjectPtr<class UAssetImportData> AssetImportData;
+	void ClearModelData();
 
 private:
 	/**
 	 * A list of string of the supported runtime, empty to support them all.
 	 */
 	TArray<FString> TargetRuntimes;
-#endif // WITH_EDITORONLY_DATA
-
-private:
 
 	/**
 	 * A string identifying the type of data inside this asset. Corresponds to the extension of the imported file.
@@ -155,6 +154,5 @@ private:
 	/**
 	 * The processed / optimized model data for the different runtimes.
 	 */
-	TMap<FString, FSharedBuffer> ModelData;
-
+	TMap<FString, TTuple<FSharedBuffer, uint32>> ModelData;
 };
diff --git a/Engine/Source/Runtime/NNE/Public/NNERuntime.h b/Engine/Source/Runtime/NNE/Public/NNERuntime.h
--- a/Engine/Source/Runtime/NNE/Public/NNERuntime.h
+++ b/Engine/Source/Runtime/NNE/Public/NNERuntime.h
@@ -53,9 +53,10 @@ public:
 	 * @param FileData The raw binary file of a neural network model.
 	 * @param FileId The unique identifier representing FileData.
 	 * @param TargetPlatform The Interface identifying the target platform for which the data needs to be created. A null pointer indicates the currently compiled/running platform.
-	 * @return Data representing the runtime specific representation of the model to be stored by UNNEModelData on success or an empty array otherwise.
+	 * @param OutMemoryAlignment Output argument of the runtime telling the caller how the returned data must be aligned in memory. A value < 1 indicates that any memory alignment will work with this model data. The variable is only set when the function succeeds.
+	 * @return Data representing the runtime specific representation of the model to be stored by UNNEModelData on success or an empty array otherwise. The data is memory aligned with OutMemoryAlignment if OutMemoryAlignment > 0.
 	 */
-	virtual TArray<uint8> CreateModelData(FString FileType, TConstArrayView<uint8> FileData, FGuid FileId, const ITargetPlatform* TargetPlatform) = 0;
+	virtual TArray<uint8> CreateModelData(FString FileType, TConstArrayView<uint8> FileData, FGuid FileId, const ITargetPlatform* TargetPlatform, uint32& OutMemoryAlignment) = 0;
 
 	/**
 	 * Get an id uniquely identifying the model data. This is used by UNNEModelData to create a FCacheKey to get and add the model data from and to a DDC.
