[NNE] Cleaning up some return types to make the NNE API more consistent

- replaced int by int32
- runtimes return shared pointers to FSharedModelData

Tested by running style transfer in editor and in a cooked build. Also ran some tests to make sure memory is never copied.

#rnx

diff --git a/Engine/Plugins/Experimental/NNERuntimeORTCpu/Source/NNERuntimeORTCpu/Private/NNERuntimeORTCpu.cpp b/Engine/Plugins/Experimental/NNERuntimeORTCpu/Source/NNERuntimeORTCpu/Private/NNERuntimeORTCpu.cpp
--- a/Engine/Plugins/Experimental/NNERuntimeORTCpu/Source/NNERuntimeORTCpu/Private/NNERuntimeORTCpu.cpp
+++ b/Engine/Plugins/Experimental/NNERuntimeORTCpu/Source/NNERuntimeORTCpu/Private/NNERuntimeORTCpu.cpp
@@ -21,7 +21,7 @@ bool UNNERuntimeORTCpuImpl::CanCreateModelData(FString FileType, TConstArrayView
 	return FileType.Compare("onnx", ESearchCase::IgnoreCase) == 0;
 }
 
-TArray<uint8> UNNERuntimeORTCpuImpl::CreateModelData(FString FileType, TConstArrayView<uint8> FileData, FGuid FileId, const ITargetPlatform* TargetPlatform, uint32& OutMemoryAlignment)
+TSharedPtr<UE::NNE::FSharedModelData> UNNERuntimeORTCpuImpl::CreateModelData(FString FileType, TConstArrayView<uint8> FileData, FGuid FileId, const ITargetPlatform* TargetPlatform)
 {
 	if (!CanCreateModelData(FileType, FileData, FileId, TargetPlatform))
 	{
@@ -48,8 +48,7 @@ TArray<uint8> UNNERuntimeORTCpuImpl::CreateModelData(FString FileType, TConstArr
 	Writer << UNNERuntimeORTCpuImpl::Version;
 	Writer.Serialize(OutputModel.Data.GetData(), OutputModel.Data.Num());
 
-	OutMemoryAlignment = 0;
-	return Result;
+	return MakeShared<UE::NNE::FSharedModelData>(MakeSharedBufferFromArray(MoveTemp(Result)), 0);
 }
 
 FString UNNERuntimeORTCpuImpl::GetModelDataIdentifier(FString FileType, TConstArrayView<uint8> FileData, FGuid FileId, const ITargetPlatform* TargetPlatform)
diff --git a/Engine/Plugins/Experimental/NNERuntimeORTCpu/Source/NNERuntimeORTCpu/Private/NNERuntimeORTCpu.h b/Engine/Plugins/Experimental/NNERuntimeORTCpu/Source/NNERuntimeORTCpu/Private/NNERuntimeORTCpu.h
--- a/Engine/Plugins/Experimental/NNERuntimeORTCpu/Source/NNERuntimeORTCpu/Private/NNERuntimeORTCpu.h
+++ b/Engine/Plugins/Experimental/NNERuntimeORTCpu/Source/NNERuntimeORTCpu/Private/NNERuntimeORTCpu.h
@@ -34,7 +34,7 @@ public:
 	virtual FString GetRuntimeName() const override { return TEXT("NNERuntimeORTCpu"); };
 
 	virtual bool CanCreateModelData(FString FileType, TConstArrayView<uint8> FileData, FGuid FileId, const ITargetPlatform* TargetPlatform) const override;
-	virtual TArray<uint8> CreateModelData(FString FileType, TConstArrayView<uint8> FileData, FGuid FileId, const ITargetPlatform* TargetPlatform, uint32& OutMemoryAlignment) override;
+	virtual TSharedPtr<UE::NNE::FSharedModelData> CreateModelData(FString FileType, TConstArrayView<uint8> FileData, FGuid FileId, const ITargetPlatform* TargetPlatform) override;
 	virtual FString GetModelDataIdentifier(FString FileType, TConstArrayView<uint8> FileData, FGuid FileId, const ITargetPlatform* TargetPlatform) override;
 
 	virtual bool CanCreateModelCPU(TObjectPtr<UNNEModelData> ModelData) const override;
diff --git a/Engine/Plugins/Experimental/NNERuntimeORTCpu/Source/NNERuntimeORTCpu/Private/NNERuntimeORTCpuModel.cpp b/Engine/Plugins/Experimental/NNERuntimeORTCpu/Source/NNERuntimeORTCpu/Private/NNERuntimeORTCpuModel.cpp
--- a/Engine/Plugins/Experimental/NNERuntimeORTCpu/Source/NNERuntimeORTCpu/Private/NNERuntimeORTCpuModel.cpp
+++ b/Engine/Plugins/Experimental/NNERuntimeORTCpu/Source/NNERuntimeORTCpu/Private/NNERuntimeORTCpuModel.cpp
@@ -162,7 +162,7 @@ namespace UE::NNERuntimeORTCpu::Private
 		return true;
 	}
 	
-	int FModelInstanceCPU::SetInputTensorShapes(TConstArrayView<NNE::FTensorShape> InInputShapes)
+	int32 FModelInstanceCPU::SetInputTensorShapes(TConstArrayView<NNE::FTensorShape> InInputShapes)
 	{
 		InputTensors.Reset();
 		OutputTensors.Reset();
@@ -175,7 +175,7 @@ namespace UE::NNERuntimeORTCpu::Private
 		}
 
 		// Setup concrete input tensor
-		for (int i = 0; i < InputSymbolicTensors.Num(); ++i)
+		for (int32 i = 0; i < InputSymbolicTensors.Num(); ++i)
 		{
 			NNE::Internal::FTensor Tensor = NNE::Internal::FTensor::Make(InputSymbolicTensors[i].GetName(), InInputShapes[i], InputSymbolicTensors[i].GetDataType());
 			InputTensors.Emplace(Tensor);
@@ -203,7 +203,7 @@ namespace UE::NNERuntimeORTCpu::Private
 		return 0;
 	}
 
-	int FModelInstanceCPU::RunSync(TConstArrayView<NNE::FTensorBindingCPU> InInputBindings, TConstArrayView<NNE::FTensorBindingCPU> InOutputBindings)
+	int32 FModelInstanceCPU::RunSync(TConstArrayView<NNE::FTensorBindingCPU> InInputBindings, TConstArrayView<NNE::FTensorBindingCPU> InOutputBindings)
 	{
 		DECLARE_SCOPE_CYCLE_COUNTER(TEXT("FModelInstanceCPU_Run"), STAT_FModelInstanceCPU_Run, STATGROUP_NNE);
 
@@ -249,7 +249,7 @@ namespace UE::NNERuntimeORTCpu::Private
 			else
 			{
 				TArray<Ort::Value> OutputOrtTensors;
-				for (int i = 0; i < InOutputBindings.Num(); ++i)
+				for (int32 i = 0; i < InOutputBindings.Num(); ++i)
 				{
 					OutputOrtTensors.Emplace(nullptr);
 				}
@@ -261,7 +261,7 @@ namespace UE::NNERuntimeORTCpu::Private
 				// Output shapes were resolved during inference: Copy the data back to bindings and expose output tensor shapes
 				CopyFromORTToBindings(OutputOrtTensors, InOutputBindings, OutputSymbolicTensors, OutputTensors);
 				check(OutputTensorShapes.IsEmpty());
-				for (int i = 0; i < OutputTensors.Num(); ++i)
+				for (int32 i = 0; i < OutputTensors.Num(); ++i)
 				{
 					OutputTensorShapes.Emplace(OutputTensors[i].GetShape());
 				}
diff --git a/Engine/Plugins/Experimental/NNERuntimeORTCpu/Source/NNERuntimeORTCpu/Private/NNERuntimeORTCpuModel.h b/Engine/Plugins/Experimental/NNERuntimeORTCpu/Source/NNERuntimeORTCpu/Private/NNERuntimeORTCpuModel.h
--- a/Engine/Plugins/Experimental/NNERuntimeORTCpu/Source/NNERuntimeORTCpu/Private/NNERuntimeORTCpuModel.h
+++ b/Engine/Plugins/Experimental/NNERuntimeORTCpu/Source/NNERuntimeORTCpu/Private/NNERuntimeORTCpuModel.h
@@ -34,8 +34,8 @@ namespace UE::NNERuntimeORTCpu::Private
 		FModelInstanceCPU(Ort::Env* InORTEnvironment, const FRuntimeConf& InRuntimeConf);
 		virtual ~FModelInstanceCPU() {};
 
-		virtual int SetInputTensorShapes(TConstArrayView<NNE::FTensorShape> InInputShapes) override;
-		virtual int RunSync(TConstArrayView<NNE::FTensorBindingCPU> InInputBindings, TConstArrayView<NNE::FTensorBindingCPU> InOutputBindings) override;
+		virtual int32 SetInputTensorShapes(TConstArrayView<NNE::FTensorShape> InInputShapes) override;
+		virtual int32 RunSync(TConstArrayView<NNE::FTensorBindingCPU> InInputBindings, TConstArrayView<NNE::FTensorBindingCPU> InOutputBindings) override;
 
 		bool Init(TConstArrayView<uint8> ModelData);
 		bool IsLoaded() const;
diff --git a/Engine/Plugins/Experimental/NNERuntimeORTCpu/Source/NNERuntimeORTCpu/Private/NNERuntimeORTCpuModule.cpp b/Engine/Plugins/Experimental/NNERuntimeORTCpu/Source/NNERuntimeORTCpu/Private/NNERuntimeORTCpuModule.cpp
--- a/Engine/Plugins/Experimental/NNERuntimeORTCpu/Source/NNERuntimeORTCpu/Private/NNERuntimeORTCpuModule.cpp
+++ b/Engine/Plugins/Experimental/NNERuntimeORTCpu/Source/NNERuntimeORTCpu/Private/NNERuntimeORTCpuModule.cpp
@@ -27,7 +27,7 @@ void FNNERuntimeORTCpuModule::ShutdownModule()
 		
 		UE::NNE::UnregisterRuntime(RuntimeCPUInterface);
 		NNERuntimeORTCpu->RemoveFromRoot();
-		NNERuntimeORTCpu = TWeakObjectPtr<UNNERuntimeORTCpuImpl>(nullptr);
+		NNERuntimeORTCpu.Reset();
 	}
 }
 
diff --git a/Engine/Plugins/Experimental/NNERuntimeORTGpu/Source/NNERuntimeORT/Private/NNERuntimeORT.cpp b/Engine/Plugins/Experimental/NNERuntimeORTGpu/Source/NNERuntimeORT/Private/NNERuntimeORT.cpp
--- a/Engine/Plugins/Experimental/NNERuntimeORTGpu/Source/NNERuntimeORT/Private/NNERuntimeORT.cpp
+++ b/Engine/Plugins/Experimental/NNERuntimeORTGpu/Source/NNERuntimeORT/Private/NNERuntimeORT.cpp
@@ -23,7 +23,7 @@ bool UNNERuntimeORTGpuImpl::CanCreateModelData(FString FileType, TConstArrayView
 	return FileType.Compare("onnx", ESearchCase::IgnoreCase) == 0;
 }
 
-TArray<uint8> UNNERuntimeORTGpuImpl::CreateModelData(FString FileType, TConstArrayView<uint8> FileData, FGuid FileId, const ITargetPlatform* TargetPlatform, uint32& OutMemoryAlignment)
+TSharedPtr<UE::NNE::FSharedModelData> UNNERuntimeORTGpuImpl::CreateModelData(FString FileType, TConstArrayView<uint8> FileData, FGuid FileId, const ITargetPlatform* TargetPlatform)
 {
 	if (!CanCreateModelData(FileType, FileData, FileId, TargetPlatform))
 	{
@@ -50,8 +50,7 @@ TArray<uint8> UNNERuntimeORTGpuImpl::CreateModelData(FString FileType, TConstArr
 	Writer << UNNERuntimeORTGpuImpl::Version;
 	Writer.Serialize(OutputModel.Data.GetData(), OutputModel.Data.Num());
 
-	OutMemoryAlignment = 0;
-	return Result;
+	return MakeShared<UE::NNE::FSharedModelData>(MakeSharedBufferFromArray(MoveTemp(Result)), 0);
 }
 
 FString UNNERuntimeORTGpuImpl::GetModelDataIdentifier(FString FileType, TConstArrayView<uint8> FileData, FGuid FileId, const ITargetPlatform* TargetPlatform)
diff --git a/Engine/Plugins/Experimental/NNERuntimeORTGpu/Source/NNERuntimeORT/Private/NNERuntimeORT.h b/Engine/Plugins/Experimental/NNERuntimeORTGpu/Source/NNERuntimeORT/Private/NNERuntimeORT.h
--- a/Engine/Plugins/Experimental/NNERuntimeORTGpu/Source/NNERuntimeORT/Private/NNERuntimeORT.h
+++ b/Engine/Plugins/Experimental/NNERuntimeORTGpu/Source/NNERuntimeORT/Private/NNERuntimeORT.h
@@ -46,7 +46,7 @@ public:
 	virtual FString GetRuntimeName() const override;
 
 	virtual bool CanCreateModelData(FString FileType, TConstArrayView<uint8> FileData, FGuid FileId, const ITargetPlatform* TargetPlatform) const override;
-	virtual TArray<uint8> CreateModelData(FString FileType, TConstArrayView<uint8> FileData, FGuid FileId, const ITargetPlatform* TargetPlatform, uint32& OutMemoryAlignment) override;
+	virtual TSharedPtr<UE::NNE::FSharedModelData> CreateModelData(FString FileType, TConstArrayView<uint8> FileData, FGuid FileId, const ITargetPlatform* TargetPlatform) override;
 	virtual FString GetModelDataIdentifier(FString FileType, TConstArrayView<uint8> FileData, FGuid FileId, const ITargetPlatform* TargetPlatform) override;
 
 	virtual bool CanCreateModelGPU(TObjectPtr<UNNEModelData> ModelData) const override;
diff --git a/Engine/Plugins/Experimental/NNERuntimeORTGpu/Source/NNERuntimeORT/Private/NNERuntimeORTModel.cpp b/Engine/Plugins/Experimental/NNERuntimeORTGpu/Source/NNERuntimeORT/Private/NNERuntimeORTModel.cpp
--- a/Engine/Plugins/Experimental/NNERuntimeORTGpu/Source/NNERuntimeORT/Private/NNERuntimeORTModel.cpp
+++ b/Engine/Plugins/Experimental/NNERuntimeORTGpu/Source/NNERuntimeORT/Private/NNERuntimeORTModel.cpp
@@ -187,7 +187,7 @@ namespace UE::NNERuntimeORT::Private
 		return true;
 	}
 	
-	int FModelInstanceORT::SetInputTensorShapes(TConstArrayView<NNE::FTensorShape> InInputShapes)
+	int32 FModelInstanceORT::SetInputTensorShapes(TConstArrayView<NNE::FTensorShape> InInputShapes)
 	{
 		InputTensors.Reset();
 		OutputTensors.Reset();
@@ -200,7 +200,7 @@ namespace UE::NNERuntimeORT::Private
 		}
 
 		// Setup concrete input tensor
-		for (int i = 0; i < InputSymbolicTensors.Num(); ++i)
+		for (int32 i = 0; i < InputSymbolicTensors.Num(); ++i)
 		{
 			NNE::Internal::FTensor Tensor = NNE::Internal::FTensor::Make(InputSymbolicTensors[i].GetName(), InInputShapes[i], InputSymbolicTensors[i].GetDataType());
 			InputTensors.Emplace(Tensor);
@@ -228,7 +228,7 @@ namespace UE::NNERuntimeORT::Private
 		return 0;
 	}
 
-	int FModelInstanceORT::RunSync(TConstArrayView<NNE::FTensorBindingGPU> InInputBindings, TConstArrayView<NNE::FTensorBindingGPU> InOutputBindings)
+	int32 FModelInstanceORT::RunSync(TConstArrayView<NNE::FTensorBindingGPU> InInputBindings, TConstArrayView<NNE::FTensorBindingGPU> InOutputBindings)
 	{
 		DECLARE_SCOPE_CYCLE_COUNTER(TEXT("FModelORT_Run"), STAT_FModelORT_Run, STATGROUP_NNE);
 
@@ -274,7 +274,7 @@ namespace UE::NNERuntimeORT::Private
 			else
 			{
 				TArray<Ort::Value> OutputOrtTensors;
-				for (int i = 0; i < InOutputBindings.Num(); ++i)
+				for (int32 i = 0; i < InOutputBindings.Num(); ++i)
 				{
 					OutputOrtTensors.Emplace(nullptr);
 				}
@@ -286,7 +286,7 @@ namespace UE::NNERuntimeORT::Private
 				// Output shapes were resolved during inference: Copy the data back to bindings and expose output tensor shapes
 				CopyFromORTToBindings(OutputOrtTensors, InOutputBindings, OutputSymbolicTensors, OutputTensors);
 				check(OutputTensorShapes.IsEmpty());
-				for (int i = 0; i < OutputTensors.Num(); ++i)
+				for (int32 i = 0; i < OutputTensors.Num(); ++i)
 				{
 					OutputTensorShapes.Emplace(OutputTensors[i].GetShape());
 				}
@@ -411,7 +411,7 @@ namespace UE::NNERuntimeORT::Private
 			}
 		}
 
-		int DeviceIndex = 0;
+		int32 DeviceIndex = 0;
 		ID3D12Device* D3D12Device = RHI->RHIGetDevice(DeviceIndex);
 
 		DML_CREATE_DEVICE_FLAGS DmlCreateFlags = DML_CREATE_DEVICE_FLAG_NONE;
@@ -466,7 +466,7 @@ namespace UE::NNERuntimeORT::Private
 		SessionOptions->EnableCpuMemArena();
 
 		//Notes: Atm we do not offer multi gpu capability/configuration
-		int DeviceId = 0;
+		int32 DeviceId = 0;
 		OrtStatusPtr Status = OrtSessionOptionsAppendExecutionProvider_CUDA(*SessionOptions.Get(), DeviceId);
 		if (Status)
 		{
diff --git a/Engine/Plugins/Experimental/NNERuntimeORTGpu/Source/NNERuntimeORT/Private/NNERuntimeORTModel.h b/Engine/Plugins/Experimental/NNERuntimeORTGpu/Source/NNERuntimeORT/Private/NNERuntimeORTModel.h
--- a/Engine/Plugins/Experimental/NNERuntimeORTGpu/Source/NNERuntimeORT/Private/NNERuntimeORTModel.h
+++ b/Engine/Plugins/Experimental/NNERuntimeORTGpu/Source/NNERuntimeORT/Private/NNERuntimeORTModel.h
@@ -34,8 +34,8 @@ namespace UE::NNERuntimeORT::Private
 		FModelInstanceORT(Ort::Env* InORTEnvironment, const FRuntimeConf& InRuntimeConf);
 		virtual ~FModelInstanceORT() = default;
 
-		virtual int SetInputTensorShapes(TConstArrayView<NNE::FTensorShape> InInputShapes) override;
-		virtual int RunSync(TConstArrayView<NNE::FTensorBindingGPU> InInputBindings, TConstArrayView<NNE::FTensorBindingGPU> InOutputBindings) override;
+		virtual int32 SetInputTensorShapes(TConstArrayView<NNE::FTensorShape> InInputShapes) override;
+		virtual int32 RunSync(TConstArrayView<NNE::FTensorBindingGPU> InInputBindings, TConstArrayView<NNE::FTensorBindingGPU> InOutputBindings) override;
 
 		bool Init(TConstArrayView<uint8> ModelData);
 		bool IsLoaded() const;
diff --git a/Engine/Plugins/Experimental/NNERuntimeORTGpu/Source/NNERuntimeORT/Private/NNERuntimeORTModule.cpp b/Engine/Plugins/Experimental/NNERuntimeORTGpu/Source/NNERuntimeORT/Private/NNERuntimeORTModule.cpp
--- a/Engine/Plugins/Experimental/NNERuntimeORTGpu/Source/NNERuntimeORT/Private/NNERuntimeORTModule.cpp
+++ b/Engine/Plugins/Experimental/NNERuntimeORTGpu/Source/NNERuntimeORT/Private/NNERuntimeORTModule.cpp
@@ -94,7 +94,7 @@ void FNNERuntimeORTModule::ShutdownModule()
 
 		UE::NNE::UnregisterRuntime(RuntimeDmlInterface);
 		NNERuntimeORTDml->RemoveFromRoot();
-		NNERuntimeORTDml = TWeakObjectPtr<UNNERuntimeORTGpuImpl>(nullptr);
+		NNERuntimeORTDml.Reset();
 	}
 
 	// NNE runtime ORT Cuda shutdown
@@ -104,7 +104,7 @@ void FNNERuntimeORTModule::ShutdownModule()
 
 		UE::NNE::UnregisterRuntime(RuntimeCudaInterface);
 		NNERuntimeORTCuda->RemoveFromRoot();
-		NNERuntimeORTCuda = TWeakObjectPtr<UNNERuntimeORTGpuImpl>(nullptr);
+		NNERuntimeORTCuda.Reset();
 	}
 
 	// Free the dll handles
diff --git a/Engine/Plugins/Experimental/NNERuntimeRDG/Source/NNERuntimeRDG/Private/NNERuntimeRDGDml.cpp b/Engine/Plugins/Experimental/NNERuntimeRDG/Source/NNERuntimeRDG/Private/NNERuntimeRDGDml.cpp
--- a/Engine/Plugins/Experimental/NNERuntimeRDG/Source/NNERuntimeRDG/Private/NNERuntimeRDGDml.cpp
+++ b/Engine/Plugins/Experimental/NNERuntimeRDG/Source/NNERuntimeRDG/Private/NNERuntimeRDGDml.cpp
@@ -101,7 +101,7 @@ bool UNNERuntimeRDGDmlImpl::CanCreateModelData(FString FileType, TConstArrayView
 	return FileType.Compare("onnx", ESearchCase::IgnoreCase) == 0;
 }
 
-TArray<uint8> UNNERuntimeRDGDmlImpl::CreateModelData(FString FileType, TConstArrayView<uint8> FileData, FGuid FileId, const ITargetPlatform* TargetPlatform, uint32& OutMemoryAlignment)
+TSharedPtr<UE::NNE::FSharedModelData> UNNERuntimeRDGDmlImpl::CreateModelData(FString FileType, TConstArrayView<uint8> FileData, FGuid FileId, const ITargetPlatform* TargetPlatform)
 {
 	if (!CanCreateModelData(FileType, FileData, FileId, TargetPlatform))
 	{
@@ -133,8 +133,7 @@ TArray<uint8> UNNERuntimeRDGDmlImpl::CreateModelData(FString FileType, TConstArr
 	Writer << Version;
 	Writer.Serialize(OutputModel.Data.GetData(), OutputModel.Data.Num());
 
-	OutMemoryAlignment = 0;
-	return Result;
+	return MakeShared<UE::NNE::FSharedModelData>(MakeSharedBufferFromArray(MoveTemp(Result)), 0);
 };
 
 FString UNNERuntimeRDGDmlImpl::GetModelDataIdentifier(FString FileType, TConstArrayView<uint8> FileData, FGuid FileId, const ITargetPlatform* TargetPlatform)
diff --git a/Engine/Plugins/Experimental/NNERuntimeRDG/Source/NNERuntimeRDG/Private/NNERuntimeRDGDml.h b/Engine/Plugins/Experimental/NNERuntimeRDG/Source/NNERuntimeRDG/Private/NNERuntimeRDGDml.h
--- a/Engine/Plugins/Experimental/NNERuntimeRDG/Source/NNERuntimeRDG/Private/NNERuntimeRDGDml.h
+++ b/Engine/Plugins/Experimental/NNERuntimeRDG/Source/NNERuntimeRDG/Private/NNERuntimeRDGDml.h
@@ -34,7 +34,7 @@ public:
 	virtual FString GetRuntimeName() const override { return TEXT("NNERuntimeRDGDml"); };
 
 	virtual bool CanCreateModelData(FString FileType, TConstArrayView<uint8> FileData, FGuid FileId, const ITargetPlatform* TargetPlatform) const override;
-	virtual TArray<uint8> CreateModelData(FString FileType, TConstArrayView<uint8> FileData, FGuid FileId, const ITargetPlatform* TargetPlatform, uint32& OutMemoryAlignment) override;
+	virtual TSharedPtr<UE::NNE::FSharedModelData> CreateModelData(FString FileType, TConstArrayView<uint8> FileData, FGuid FileId, const ITargetPlatform* TargetPlatform) override;
 	virtual FString GetModelDataIdentifier(FString FileType, TConstArrayView<uint8> FileData, FGuid FileId, const ITargetPlatform* TargetPlatform) override;
 
 	virtual bool CanCreateModelRDG(TObjectPtr<UNNEModelData> ModelData) const override;
diff --git a/Engine/Plugins/Experimental/NNERuntimeRDG/Source/NNERuntimeRDG/Private/NNERuntimeRDGHlsl.cpp b/Engine/Plugins/Experimental/NNERuntimeRDG/Source/NNERuntimeRDG/Private/NNERuntimeRDGHlsl.cpp
--- a/Engine/Plugins/Experimental/NNERuntimeRDG/Source/NNERuntimeRDG/Private/NNERuntimeRDGHlsl.cpp
+++ b/Engine/Plugins/Experimental/NNERuntimeRDG/Source/NNERuntimeRDG/Private/NNERuntimeRDGHlsl.cpp
@@ -100,7 +100,7 @@ bool UNNERuntimeRDGHlslImpl::CanCreateModelRDG(TObjectPtr<UNNEModelData> ModelDa
 	return bResult;
 };
 
-TArray<uint8> UNNERuntimeRDGHlslImpl::CreateModelData(FString FileType, TConstArrayView<uint8> FileData, FGuid FileId, const ITargetPlatform* TargetPlatform, uint32& OutMemoryAlignment)
+TSharedPtr<UE::NNE::FSharedModelData> UNNERuntimeRDGHlslImpl::CreateModelData(FString FileType, TConstArrayView<uint8> FileData, FGuid FileId, const ITargetPlatform* TargetPlatform)
 {
 	if (!CanCreateModelData(FileType, FileData, FileId, TargetPlatform))
 	{
@@ -127,8 +127,7 @@ TArray<uint8> UNNERuntimeRDGHlslImpl::CreateModelData(FString FileType, TConstAr
 	Writer << Version;
 	Writer.Serialize(OutputModel.Data.GetData(), OutputModel.Data.Num());
 
-	OutMemoryAlignment = 0;
-	return Result;
+	return MakeShared<UE::NNE::FSharedModelData>(MakeSharedBufferFromArray(MoveTemp(Result)), 0);
 };
 
 FString UNNERuntimeRDGHlslImpl::GetModelDataIdentifier(FString FileType, TConstArrayView<uint8> FileData, FGuid FileId, const ITargetPlatform* TargetPlatform)
diff --git a/Engine/Plugins/Experimental/NNERuntimeRDG/Source/NNERuntimeRDG/Private/NNERuntimeRDGHlsl.h b/Engine/Plugins/Experimental/NNERuntimeRDG/Source/NNERuntimeRDG/Private/NNERuntimeRDGHlsl.h
--- a/Engine/Plugins/Experimental/NNERuntimeRDG/Source/NNERuntimeRDG/Private/NNERuntimeRDGHlsl.h
+++ b/Engine/Plugins/Experimental/NNERuntimeRDG/Source/NNERuntimeRDG/Private/NNERuntimeRDGHlsl.h
@@ -29,7 +29,7 @@ public:
 	virtual FString GetRuntimeName() const override { return TEXT("NNERuntimeRDGHlsl"); };
 
 	virtual bool CanCreateModelData(FString FileType, TConstArrayView<uint8> FileData, FGuid FileId, const ITargetPlatform* TargetPlatform) const override;
-	virtual TArray<uint8> CreateModelData(FString FileType, TConstArrayView<uint8> FileData, FGuid FileId, const ITargetPlatform* TargetPlatform, uint32& OutMemoryAlignment) override;
+	virtual TSharedPtr<UE::NNE::FSharedModelData> CreateModelData(FString FileType, TConstArrayView<uint8> FileData, FGuid FileId, const ITargetPlatform* TargetPlatform) override;
 	virtual FString GetModelDataIdentifier(FString FileType, TConstArrayView<uint8> FileData, FGuid FileId, const ITargetPlatform* TargetPlatform) override;
 
 	virtual bool CanCreateModelRDG(TObjectPtr<UNNEModelData> ModelData) const override;
diff --git a/Engine/Plugins/Experimental/NNERuntimeRDG/Source/NNERuntimeRDG/Private/NNERuntimeRDGModel.cpp b/Engine/Plugins/Experimental/NNERuntimeRDG/Source/NNERuntimeRDG/Private/NNERuntimeRDGModel.cpp
--- a/Engine/Plugins/Experimental/NNERuntimeRDG/Source/NNERuntimeRDG/Private/NNERuntimeRDGModel.cpp
+++ b/Engine/Plugins/Experimental/NNERuntimeRDG/Source/NNERuntimeRDG/Private/NNERuntimeRDGModel.cpp
@@ -96,7 +96,7 @@ bool FModelInstanceRDG::LoadModel(TConstArrayView<uint8> ModelData, FNNERuntimeF
 
 
 
-int FModelInstanceRDG::SetInputTensorShapes(TConstArrayView<NNE::FTensorShape> InInputShapes)
+int32 FModelInstanceRDG::SetInputTensorShapes(TConstArrayView<NNE::FTensorShape> InInputShapes)
 {
 	OutputTensorShapes.Reset(OutputTensorIndices.Num());
 
@@ -150,7 +150,7 @@ int FModelInstanceRDG::SetInputTensorShapes(TConstArrayView<NNE::FTensorShape> I
 	}
 
 	checkCode(
-		for (int i = 0; i < AllTensorRDGRefs.Num(); ++i)
+		for (int32 i = 0; i < AllTensorRDGRefs.Num(); ++i)
 		{
 			checkf(AllTensorRDGRefs[i] != nullptr, TEXT("Tensor at index %d, was not allocated for model preparation."), i);
 		};
@@ -163,7 +163,7 @@ int FModelInstanceRDG::SetInputTensorShapes(TConstArrayView<NNE::FTensorShape> I
 	}
 
 	checkCode(
-		for (int i = 0; i < AllTensorRDGRefs.Num(); ++i)
+		for (int32 i = 0; i < AllTensorRDGRefs.Num(); ++i)
 		{
 			checkf(AllTensorRDGRefs[i] != nullptr, TEXT("Tensor at index %d, was not allocated after model preparation."), i);
 			checkf(AllTensorRDGRefs[i]->GetShape().IsCompatibleWith(AllSymbolicTensorDescs[i].GetShape()), TEXT("Tensor at index %d have a shape incompatible with model definition."), i);
@@ -197,11 +197,11 @@ FRDGBufferDesc CreateRDGBufferDescForTensorRDG(const FTensorRDG& Tensor)
 /**
  * Enqueue operators to RDG, the caller will run the GraphBuilder.Execute()
  */
-int FModelInstanceRDG::EnqueueRDG(FRDGBuilder& RDGBuilder, TConstArrayView<NNE::FTensorBindingRDG> InInputBindings, TConstArrayView<NNE::FTensorBindingRDG> InOutputBindings)
+int32 FModelInstanceRDG::EnqueueRDG(FRDGBuilder& RDGBuilder, TConstArrayView<NNE::FTensorBindingRDG> InInputBindings, TConstArrayView<NNE::FTensorBindingRDG> InOutputBindings)
 {
 	check(IsInRenderingThread());
 
-	int Res;
+	int32 Res;
 
 	// Verify the model inputs were prepared
 	if (InputTensorShapes.Num() == 0)
@@ -253,7 +253,7 @@ int FModelInstanceRDG::EnqueueRDG(FRDGBuilder& RDGBuilder, TConstArrayView<NNE::
 	return 0;
 }
 
-int FModelInstanceRDG::SetTensors(FRDGBuilder& GraphBuilder, FTensorRDGArray& InTensorRDGs, TConstArrayView<NNE::FTensorBindingRDG> InBindings)
+int32 FModelInstanceRDG::SetTensors(FRDGBuilder& GraphBuilder, FTensorRDGArray& InTensorRDGs, TConstArrayView<NNE::FTensorBindingRDG> InBindings)
 {
 	check(InBindings.Num() == InTensorRDGs.Num());
 	
diff --git a/Engine/Plugins/Experimental/NNERuntimeRDG/Source/NNERuntimeRDG/Private/NNERuntimeRDGModel.h b/Engine/Plugins/Experimental/NNERuntimeRDG/Source/NNERuntimeRDG/Private/NNERuntimeRDGModel.h
--- a/Engine/Plugins/Experimental/NNERuntimeRDG/Source/NNERuntimeRDG/Private/NNERuntimeRDGModel.h
+++ b/Engine/Plugins/Experimental/NNERuntimeRDG/Source/NNERuntimeRDG/Private/NNERuntimeRDGModel.h
@@ -17,12 +17,12 @@ namespace UE::NNERuntimeRDG::Private
 		FModelInstanceRDG() {};
 		virtual ~FModelInstanceRDG() = default;
 
-		virtual int SetInputTensorShapes(TConstArrayView<NNE::FTensorShape> InInputShapes) override;
-		virtual int EnqueueRDG(FRDGBuilder& RDGBuilder, TConstArrayView<NNE::FTensorBindingRDG> InInputBindings, TConstArrayView<NNE::FTensorBindingRDG> InOutputBindings) override;
+		virtual int32 SetInputTensorShapes(TConstArrayView<NNE::FTensorShape> InInputShapes) override;
+		virtual int32 EnqueueRDG(FRDGBuilder& RDGBuilder, TConstArrayView<NNE::FTensorBindingRDG> InInputBindings, TConstArrayView<NNE::FTensorBindingRDG> InOutputBindings) override;
 		
 	protected:
 		bool LoadModel(TConstArrayView<uint8> ModelData, FNNERuntimeFormat& Format, int32 GuidAndVersionSize);
-		int SetTensors(FRDGBuilder& GraphBuilder, FTensorRDGArray& InTensorRDGs, TConstArrayView<NNE::FTensorBindingRDG> InBindings);
+		int32 SetTensors(FRDGBuilder& GraphBuilder, FTensorRDGArray& InTensorRDGs, TConstArrayView<NNE::FTensorBindingRDG> InBindings);
 
 		virtual int PrepareTensorShapesAndData() = 0;
 		virtual bool PrepareModelRDG(FRDGBuilder& RDGBuilder) { return false; }
diff --git a/Engine/Plugins/Experimental/NNERuntimeRDG/Source/NNERuntimeRDG/Private/NNERuntimeRDGModule.cpp b/Engine/Plugins/Experimental/NNERuntimeRDG/Source/NNERuntimeRDG/Private/NNERuntimeRDGModule.cpp
--- a/Engine/Plugins/Experimental/NNERuntimeRDG/Source/NNERuntimeRDG/Private/NNERuntimeRDGModule.cpp
+++ b/Engine/Plugins/Experimental/NNERuntimeRDG/Source/NNERuntimeRDG/Private/NNERuntimeRDGModule.cpp
@@ -44,7 +44,7 @@ void FNNERuntimeRDGModule::ShutdownModule()
 
 		UE::NNE::UnregisterRuntime(RuntimeCPUInterface);
 		NNERuntimeRDGHlsl->RemoveFromRoot();
-		NNERuntimeRDGHlsl = TWeakObjectPtr<UNNERuntimeRDGHlslImpl>(nullptr);
+		NNERuntimeRDGHlsl.Reset();
 	}
 
 #ifdef NNE_USE_DIRECTML
@@ -55,7 +55,7 @@ void FNNERuntimeRDGModule::ShutdownModule()
 
 		UE::NNE::UnregisterRuntime(RuntimeDmlInterface);
 		NNERuntimeRDGDml->RemoveFromRoot();
-		NNERuntimeRDGDml = TWeakObjectPtr<UNNERuntimeRDGDmlImpl>(nullptr);
+		NNERuntimeRDGDml.Reset();
 	}
 #endif
 
diff --git a/Engine/Source/Runtime/NNE/Internal/NNEModelBase.h b/Engine/Source/Runtime/NNE/Internal/NNEModelBase.h
--- a/Engine/Source/Runtime/NNE/Internal/NNEModelBase.h
+++ b/Engine/Source/Runtime/NNE/Internal/NNEModelBase.h
@@ -18,7 +18,7 @@ namespace UE::NNE::Internal
 		virtual TConstArrayView<NNE::FTensorDesc> GetOutputTensorDescs() const override;
 		virtual TConstArrayView<NNE::FTensorShape> GetInputTensorShapes() const override;
 		virtual TConstArrayView<NNE::FTensorShape> GetOutputTensorShapes() const override;
-		virtual int SetInputTensorShapes(TConstArrayView<NNE::FTensorShape> InInputShapes) override;
+		virtual int32 SetInputTensorShapes(TConstArrayView<NNE::FTensorShape> InInputShapes) override;
 
 	protected:
 
@@ -50,7 +50,7 @@ namespace UE::NNE::Internal
 		return OutputTensorShapes;
 	}
 
-	template <class T> int FModelInstanceBase<T>::SetInputTensorShapes(TConstArrayView<FTensorShape> InInputShapes)
+	template <class T> int32 FModelInstanceBase<T>::SetInputTensorShapes(TConstArrayView<FTensorShape> InInputShapes)
 	{
 		InputTensorShapes.Reset(InInputShapes.Num());
 
diff --git a/Engine/Source/Runtime/NNE/Private/NNEModelData.cpp b/Engine/Source/Runtime/NNE/Private/NNEModelData.cpp
--- a/Engine/Source/Runtime/NNE/Private/NNEModelData.cpp
+++ b/Engine/Source/Runtime/NNE/Private/NNEModelData.cpp
@@ -70,8 +70,11 @@ namespace UE::NNE::ModelData
 		return { UE::DerivedData::FCacheBucket(FWideStringView(*FileId)), FIoHash::HashBuffer(MakeMemoryView(FTCHARToUTF8(RequestId))) };
 	}
 
-	inline void PutIntoDDC(const FGuid& FileId, const FString& RuntimeName, const FString& ModelDataIdentifier, FSharedBuffer& Data, uint32 MemoryAlignment)
+	inline void PutIntoDDC(const FGuid& FileId, const FString& RuntimeName, const FString& ModelDataIdentifier, TSharedPtr<UE::NNE::FSharedModelData> Data)
 	{
+		check(Data.IsValid());
+		check(Data->GetView().Num() > 0);
+
 		FString FileIdStr = FileId.ToString(EGuidFormats::Digits);
 		FString RequestId = GetDDCRequestId(FileIdStr, RuntimeName, ModelDataIdentifier);
 
@@ -80,14 +83,14 @@ namespace UE::NNE::ModelData
 
 		Requests[0].Name = FString("Put-") + RequestId;
 		Requests[0].Key = CreateCacheKey(FileIdStr, RequestId);
-		Requests[0].Value = UE::DerivedData::FValue::Compress(FCompositeBuffer(MakeSharedBufferFromArray(TArray<uint32>({ MemoryAlignment })), Data));
+		Requests[0].Value = UE::DerivedData::FValue::Compress(FCompositeBuffer(MakeSharedBufferFromArray(TArray<uint32>({ Data->GetMemoryAlignment() })), FSharedBuffer::MakeView(Data->GetView().GetData(), Data->GetView().Num())));
 
 		UE::DerivedData::FRequestOwner BlockingPutOwner(UE::DerivedData::EPriority::Blocking);
 		UE::DerivedData::GetCache().PutValue(Requests, BlockingPutOwner);
 		BlockingPutOwner.Wait();
 	}
 
-	inline FSharedBuffer GetFromDDC(const FGuid& FileId, const FString& RuntimeName, const FString& ModelDataIdentifier, uint32& OutMemoryAlignment)
+	inline TSharedPtr<UE::NNE::FSharedModelData> GetFromDDC(const FGuid& FileId, const FString& RuntimeName, const FString& ModelDataIdentifier)
 	{
 		FString FileIdStr = FileId.ToString(EGuidFormats::Digits);
 		FString RequestId = GetDDCRequestId(FileIdStr, RuntimeName, ModelDataIdentifier);
@@ -98,9 +101,9 @@ namespace UE::NNE::ModelData
 		Requests[0].Name = FString("Get-") + RequestId;
 		Requests[0].Key = CreateCacheKey(FileIdStr, RequestId);
 
-		FSharedBuffer Result;
+		TSharedPtr<UE::NNE::FSharedModelData> Result;
 		UE::DerivedData::FRequestOwner BlockingGetOwner(UE::DerivedData::EPriority::Blocking);
-		UE::DerivedData::GetCache().GetValue(Requests, BlockingGetOwner, [&Result, &OutMemoryAlignment](UE::DerivedData::FCacheGetValueResponse&& Response)
+		UE::DerivedData::GetCache().GetValue(Requests, BlockingGetOwner, [&Result](UE::DerivedData::FCacheGetValueResponse&& Response)
 		{
 			if (Response.Value.HasData() && Response.Value.GetRawSize() > sizeof(uint32))
 			{
@@ -111,8 +114,7 @@ namespace UE::NNE::ModelData
 				void* Data = FMemory::Malloc(DataSize, MemoryAlignment);
 				if (Reader.TryDecompressTo(FMutableMemoryView(Data, DataSize), sizeof(uint32)))
 				{
-					OutMemoryAlignment = MemoryAlignment;
-					Result = FSharedBuffer::TakeOwnership(Data, DataSize, FMemory::Free);
+					Result = MakeShared<UE::NNE::FSharedModelData>(FSharedBuffer::TakeOwnership(Data, DataSize, FMemory::Free), MemoryAlignment);
 				}
 				else
 				{
@@ -126,24 +128,14 @@ namespace UE::NNE::ModelData
 
 #endif // WITH_EDITOR
 
-	inline FSharedBuffer CreateModelData(const FString& RuntimeName, FString FileType, const TArray<uint8>& FileData, FGuid FileId, const ITargetPlatform* TargetPlatform, uint32& OutMemoryAlignment)
+	inline TSharedPtr<UE::NNE::FSharedModelData> CreateModelData(const FString& RuntimeName, FString FileType, const TArray<uint8>& FileData, FGuid FileId, const ITargetPlatform* TargetPlatform)
 	{
 		TWeakInterfacePtr<INNERuntime> NNERuntime = UE::NNE::GetRuntime<INNERuntime>(RuntimeName);
 		if (NNERuntime.IsValid())
 		{
 			if (NNERuntime->CanCreateModelData(FileType, FileData, FileId, TargetPlatform))
 			{
-				uint32 TmpOutMemoryAlignment = 0;
-				TArray<uint8> ModelData = NNERuntime->CreateModelData(FileType, FileData, FileId, TargetPlatform, TmpOutMemoryAlignment);
-
-				// Make sure a runtime does fulfill its own alignment requirement
-				checkf(TmpOutMemoryAlignment <= 1 || (((uintptr_t)(const void *)(ModelData.GetData())) % TmpOutMemoryAlignment == 0), TEXT("Runtimes must return ModelData that is aligned with OutMemoryAlignment!"))
-
-				if (ModelData.Num() > 0)
-				{
-					OutMemoryAlignment = TmpOutMemoryAlignment;
-					return MakeSharedBufferFromArray(MoveTemp(ModelData));
-				}
+				return NNERuntime->CreateModelData(FileType, FileData, FileId, TargetPlatform);
 			}
 		}
 		else
@@ -155,11 +147,35 @@ namespace UE::NNE::ModelData
 				UE_LOG(LogNNE, Error, TEXT("- %s"), *Runtimes[i]->GetRuntimeName());
 			}
 		}
-		return {};
+		return TSharedPtr<UE::NNE::FSharedModelData>();
 	}
 
 } // UE::NNE::ModelData
 
+namespace UE::NNE
+{
+	FSharedModelData::FSharedModelData(FSharedBuffer InData, uint32 InMemoryAlignment) : Data(InData), MemoryAlignment(InMemoryAlignment)
+	{
+		checkf(Data.IsOwned(), TEXT("InData data must be ownned!"));
+		checkf(MemoryAlignment <= 1 || (((uintptr_t)(const void*)(InData.GetData())) % MemoryAlignment == 0), TEXT("InData must be aligned with InMemoryAlignment!"))
+	}
+
+	FSharedModelData::FSharedModelData()
+	{
+		MemoryAlignment = 0;
+	}
+
+	TConstArrayView<uint8> FSharedModelData::GetView() const
+	{
+		return MakeArrayView(static_cast<const uint8*>(Data.GetData()), Data.GetSize());
+	}
+
+	uint32 FSharedModelData::GetMemoryAlignment() const
+	{
+		return MemoryAlignment;
+	}
+} // UE::NNE
+
 void UNNEModelData::GetAssetRegistryTags(TArray<FAssetRegistryTag>& OutTags) const
 {
 	OutTags.Add(FAssetRegistryTag("TargetRuntimes", UE::NNE::ModelData::GetRuntimesAsString(GetTargetRuntimes()), FAssetRegistryTag::TT_Alphabetical));
@@ -206,11 +222,10 @@ void UNNEModelData::Serialize(FArchive& Ar)
 
 				for (const FString& RuntimeName : CookRuntimeNames)
 				{
-					uint32 MemoryAlignment = 0;
-					FSharedBuffer CreatedData = UE::NNE::ModelData::CreateModelData(RuntimeName, FileType, FileData, FileId, Ar.GetArchiveState().CookingTarget(), MemoryAlignment);
-					if (CreatedData.GetSize() > 0)
+					TSharedPtr<UE::NNE::FSharedModelData> CreatedData = UE::NNE::ModelData::CreateModelData(RuntimeName, FileType, FileData, FileId, Ar.GetArchiveState().CookingTarget());
+					if (CreatedData.IsValid() && CreatedData->GetView().Num() > 0)
 					{
-						ModelData.Add(RuntimeName, MakeTuple(CreatedData, MemoryAlignment));
+						ModelData.Add(RuntimeName, CreatedData);
 #if WITH_EDITOR
 						TWeakInterfacePtr<INNERuntime> NNERuntime = UE::NNE::GetRuntime<INNERuntime>(RuntimeName);
 						if (NNERuntime.IsValid())
@@ -218,7 +233,7 @@ void UNNEModelData::Serialize(FArchive& Ar)
 							FString ModelDataIdentifier = NNERuntime->GetModelDataIdentifier(FileType, FileData, FileId, Ar.GetArchiveState().CookingTarget());
 							if (ModelDataIdentifier.Len() > 0)
 							{
-								UE::NNE::ModelData::PutIntoDDC(FileId, RuntimeName, ModelDataIdentifier, CreatedData, MemoryAlignment);
+								UE::NNE::ModelData::PutIntoDDC(FileId, RuntimeName, ModelDataIdentifier, CreatedData);
 							}
 							else
 							{
@@ -259,13 +274,13 @@ void UNNEModelData::Serialize(FArchive& Ar)
 			{
 				Ar << RuntimeNames[i];
 
-				uint32 MemoryAlignment = ModelData[RuntimeNames[i]].Get<1>();
+				uint32 MemoryAlignment = ModelData[RuntimeNames[i]]->GetMemoryAlignment();
 				Ar << MemoryAlignment;
 
-				uint64 DataSize = ModelData[RuntimeNames[i]].Get<0>().GetSize();
+				uint64 DataSize = ModelData[RuntimeNames[i]]->GetView().Num();
 				Ar << DataSize;
 
-				Ar.Serialize((void*)ModelData[RuntimeNames[i]].Get<0>().GetData(), DataSize);
+				Ar.Serialize((void*)ModelData[RuntimeNames[i]]->GetView().GetData(), DataSize);
 			}
 		}
 		else
@@ -299,7 +314,7 @@ void UNNEModelData::Serialize(FArchive& Ar)
 			{
 				Ar << Name;
 				Ar << Data;
-				ModelData.Add(Name, MakeTuple(MakeSharedBufferFromArray(MoveTemp(Data)), 0));
+				ModelData.Add(Name, MakeShared<UE::NNE::FSharedModelData>(MakeSharedBufferFromArray(MoveTemp(Data)), 0));
 			}
 			UE_LOG(LogNNE, Warning, TEXT("[DEPRECATION] UNNEModelData: The asset %s (v0) is deprecated. Please right-click the asset and select 'Save' to update it to the latest version."), *this->GetName());
 			break;
@@ -319,7 +334,7 @@ void UNNEModelData::Serialize(FArchive& Ar)
 			{
 				Ar << Name;
 				Ar << Data;
-				ModelData.Add(Name, MakeTuple(MakeSharedBufferFromArray(MoveTemp(Data)), 0));
+				ModelData.Add(Name, MakeShared<UE::NNE::FSharedModelData>(MakeSharedBufferFromArray(MoveTemp(Data)), 0));
 			}
 			UE_LOG(LogNNE, Warning, TEXT("[DEPRECATION] UNNEModelData: The asset %s (v1) is deprecated. Please right-click the asset and select 'Save' to update it to the latest version."), *this->GetName());
 			break;
@@ -337,7 +352,7 @@ void UNNEModelData::Serialize(FArchive& Ar)
 				Ar << DataSize;
 				RawData = FMemory::Malloc(DataSize, MemoryAlignment);
 				Ar.Serialize(RawData, DataSize);
-				ModelData.Add(Name, MakeTuple(FSharedBuffer::TakeOwnership(RawData, DataSize, FMemory::Free), MemoryAlignment));
+				ModelData.Add(Name, MakeShared<UE::NNE::FSharedModelData>(FSharedBuffer::TakeOwnership(RawData, DataSize, FMemory::Free), MemoryAlignment));
 			}
 			break;
 
@@ -381,12 +396,12 @@ void UNNEModelData::SetTargetRuntimes(TArrayView<const FString> RuntimeNames)
 	}
 }
 
-FString UNNEModelData::GetFileType()
+FString UNNEModelData::GetFileType() const
 {
 	return FileType;
 }
 
-TConstArrayView<uint8> UNNEModelData::GetFileData()
+TConstArrayView<uint8> UNNEModelData::GetFileData() const
 {
 	return FileData;
 }
@@ -397,7 +412,7 @@ void UNNEModelData::ClearFileDataAndFileType()
 	FileData.Empty();
 }
 
-FGuid UNNEModelData::GetFileId()
+FGuid UNNEModelData::GetFileId() const
 {
 	return FileId;
 }
@@ -417,10 +432,10 @@ TSharedPtr<UE::NNE::FSharedModelData> UNNEModelData::GetModelData(const FString&
 	}
 
 	// Check if we have a local cache hit
-	TTuple<FSharedBuffer, uint32>* LocalData = ModelData.Find(RuntimeName);
-	if (LocalData)
+	TSharedPtr<UE::NNE::FSharedModelData>* LocalDataPtr = ModelData.Find(RuntimeName);
+	if (LocalDataPtr)
 	{
-		return MakeShared<UE::NNE::FSharedModelData>(LocalData->Get<0>());
+		return *LocalDataPtr;
 	}
 
 	// After this point FileData is required to either get the cache id or recreate it from scratch
@@ -446,33 +461,30 @@ TSharedPtr<UE::NNE::FSharedModelData> UNNEModelData::GetModelData(const FString&
 	}
 
 	// Check if we have a DDC cache hit
-	uint32 RemoteMemoryAlignment = 0;
-	FSharedBuffer RemoteData = UE::NNE::ModelData::GetFromDDC(FileId, RuntimeName, ModelDataIdentifier, RemoteMemoryAlignment);
-	if (RemoteData.GetSize() > 0)
+	TSharedPtr<UE::NNE::FSharedModelData> RemoteData = UE::NNE::ModelData::GetFromDDC(FileId, RuntimeName, ModelDataIdentifier);
+	if (RemoteData.IsValid() && RemoteData->GetView().Num() > 0)
 	{
-		ModelData.Add(RuntimeName, MakeTuple(RemoteData, RemoteMemoryAlignment));
-
-		return MakeShared<UE::NNE::FSharedModelData>(RemoteData);
+		ModelData.Add(RuntimeName, RemoteData);
+		return RemoteData;
 	}
 #endif //WITH_EDITOR
 
 	// Try to create the model
-	uint32 CreatedMemoryAlignment = 0;
-	FSharedBuffer CreatedData = UE::NNE::ModelData::CreateModelData(RuntimeName, FileType, FileData, FileId, nullptr, CreatedMemoryAlignment);
-	if (CreatedData.GetSize() < 1)
+	TSharedPtr<UE::NNE::FSharedModelData> CreatedData = UE::NNE::ModelData::CreateModelData(RuntimeName, FileType, FileData, FileId, nullptr);
+	if (!CreatedData.IsValid() || CreatedData->GetView().Num() < 1)
 	{
 		return TSharedPtr<UE::NNE::FSharedModelData>();
 	}
 
 	// Cache the model
-	ModelData.Add(RuntimeName, MakeTuple(CreatedData, CreatedMemoryAlignment));
+	ModelData.Add(RuntimeName, CreatedData);
 
 #if WITH_EDITOR
 	// And put it into DDC
-	UE::NNE::ModelData::PutIntoDDC(FileId, RuntimeName, ModelDataIdentifier, CreatedData, CreatedMemoryAlignment);
+	UE::NNE::ModelData::PutIntoDDC(FileId, RuntimeName, ModelDataIdentifier, CreatedData);
 #endif //WITH_EDITOR
 
-	return MakeShared<UE::NNE::FSharedModelData>(CreatedData);
+	return CreatedData;
 }
 
 void UNNEModelData::ClearModelData()
diff --git a/Engine/Source/Runtime/NNE/Public/NNEModelData.h b/Engine/Source/Runtime/NNE/Public/NNEModelData.h
--- a/Engine/Source/Runtime/NNE/Public/NNEModelData.h
+++ b/Engine/Source/Runtime/NNE/Public/NNEModelData.h
@@ -12,27 +12,52 @@
 namespace UE::NNE
 {
 	/**
-	 * This class present a ref counted view on an immutable memory buffer. 
-	 * Allowing runtime to reference result of GetModelData() even if they outlive UNNEModelData.
+	 * This class implements a reference counted view on an immutable memory buffer representing model data.
+	 * 
+	 * It allows runtimes to reference results of GetModelData() even if they outlive UNNEModelData.
 	 */
-	class FSharedModelData
+	class NNE_API FSharedModelData
 	{
-	private:
-		FSharedBuffer Data;
 
 	public:
-		FSharedModelData(FSharedBuffer InData) : Data(InData) {}
-		FSharedModelData() {}
+		/**
+		 * Constructor to shared model data.
+		 * 
+		 * @param InData The shared buffer containing the model data. InData must be owned and the memory aligned with InMemoryAlignment.
+		 * @param InMemoryAlignment The memory alignment with which InData has been aligned. A value <= 1 indicates arbitrary memory alignment.
+		 */
+		FSharedModelData(FSharedBuffer InData, uint32 InMemoryAlignment);
+
+		/**
+		 * Constructor to create empty data.
+		 */
+		FSharedModelData();
 
 		/**
 		 * Get a const array view on the shared data which is guaranteed to remain valid as long as this objects exists.
 		 *
 		 * @return A const array view of the shared data.
 		 */
-		TConstArrayView<uint8> GetView() const
-		{
-			return MakeArrayView(static_cast<const uint8*>(Data.GetData()), Data.GetSize());
-		}
+		TConstArrayView<uint8> GetView() const;
+
+		/**
+		 * Get the memory alignment with which the data has been aligned.
+		 *
+		 * @return Memory alignment with which the data has been aligned. A value <= 1 indicates arbitrary memory alignment.
+		 */
+		uint32 GetMemoryAlignment() const;
+
+	private:
+
+		/**
+		 * The shared buffer containing the model data. Data must be aligned with MemoryAlignment.
+		 */
+		FSharedBuffer Data;
+		
+		/**
+		 * The memory alignment with which Data has been aligned. A value <= 1 indicates arbitrary memory alignment.
+		 */
+		uint32 MemoryAlignment;
 	};
 }
 
@@ -86,7 +111,7 @@ public:
 	 *
 	 * @return The FileType.
 	 */
-	FString GetFileType();
+	FString GetFileType() const;
 
 	/**
 	 * Get read only access to FileData.
@@ -95,7 +120,7 @@ public:
 	 *
 	 * @return The FileData.
 	 */
-	TConstArrayView<uint8> GetFileData();
+	TConstArrayView<uint8> GetFileData() const;
 
 	/**
 	 * Clears the FileData and the FileType.
@@ -111,7 +136,7 @@ public:
 	 *
 	 * @return The FileId.
 	 */
-	FGuid GetFileId();
+	FGuid GetFileId() const;
 
 	/**
 	 * Get the cached (editor) or cooked (game) optimized model data for a given runtime.
@@ -154,5 +179,5 @@ private:
 	/**
 	 * The processed / optimized model data for the different runtimes.
 	 */
-	TMap<FString, TTuple<FSharedBuffer, uint32>> ModelData;
+	TMap<FString, TSharedPtr<UE::NNE::FSharedModelData>> ModelData;
 };
diff --git a/Engine/Source/Runtime/NNE/Public/NNERuntime.h b/Engine/Source/Runtime/NNE/Public/NNERuntime.h
--- a/Engine/Source/Runtime/NNE/Public/NNERuntime.h
+++ b/Engine/Source/Runtime/NNE/Public/NNERuntime.h
@@ -9,6 +9,11 @@
 
 class ITargetPlatform;
 
+namespace UE::NNE
+{
+	class FSharedModelData;
+}
+
 UINTERFACE()
 class NNE_API UNNERuntime : public UInterface
 {
@@ -53,10 +58,9 @@ public:
 	 * @param FileData The raw binary file of a neural network model.
 	 * @param FileId The unique identifier representing FileData.
 	 * @param TargetPlatform The Interface identifying the target platform for which the data needs to be created. A null pointer indicates the currently compiled/running platform.
-	 * @param OutMemoryAlignment Output argument of the runtime telling the caller how the returned data must be aligned in memory. A value < 1 indicates that any memory alignment will work with this model data. The variable is only set when the function succeeds.
-	 * @return Data representing the runtime specific representation of the model to be stored by UNNEModelData on success or an empty array otherwise. The data is memory aligned with OutMemoryAlignment if OutMemoryAlignment > 0.
+	 * @return A pointer to data representing the runtime specific representation of the model to be stored by UNNEModelData on success or an invalid pointer otherwise.
 	 */
-	virtual TArray<uint8> CreateModelData(FString FileType, TConstArrayView<uint8> FileData, FGuid FileId, const ITargetPlatform* TargetPlatform, uint32& OutMemoryAlignment) = 0;
+	virtual TSharedPtr<UE::NNE::FSharedModelData> CreateModelData(FString FileType, TConstArrayView<uint8> FileData, FGuid FileId, const ITargetPlatform* TargetPlatform) = 0;
 
 	/**
 	 * Get an id uniquely identifying the model data. This is used by UNNEModelData to create a FCacheKey to get and add the model data from and to a DDC.
diff --git a/Engine/Source/Runtime/NNE/Public/NNERuntimeCPU.h b/Engine/Source/Runtime/NNE/Public/NNERuntimeCPU.h
--- a/Engine/Source/Runtime/NNE/Public/NNERuntimeCPU.h
+++ b/Engine/Source/Runtime/NNE/Public/NNERuntimeCPU.h
@@ -80,7 +80,7 @@ public:
 	 * @param InInputShapes The input shapes to prepare the model with.
 	 * @return 0 on success or a non-zero number otherwise.
 	 */
-	virtual int SetInputTensorShapes(TConstArrayView<FTensorShape> InInputShapes) = 0;
+	virtual int32 SetInputTensorShapes(TConstArrayView<FTensorShape> InInputShapes) = 0;
 
 	/**
 	 * Evaluate the model synchronously.
@@ -94,7 +94,7 @@ public:
 	 * @param InOutputTensors An array containing tensor bindings for each output tensor with caller owned memory big enough to contain the results on success.
 	 * @return 0 on success or a non-zero number otherwise.
 	 */
-	virtual int RunSync(TConstArrayView<FTensorBindingCPU> InInputTensors, TConstArrayView<FTensorBindingCPU> InOutputTensors) = 0;
+	virtual int32 RunSync(TConstArrayView<FTensorBindingCPU> InInputTensors, TConstArrayView<FTensorBindingCPU> InOutputTensors) = 0;
 };
 
 /**
diff --git a/Engine/Source/Runtime/NNE/Public/NNERuntimeGPU.h b/Engine/Source/Runtime/NNE/Public/NNERuntimeGPU.h
--- a/Engine/Source/Runtime/NNE/Public/NNERuntimeGPU.h
+++ b/Engine/Source/Runtime/NNE/Public/NNERuntimeGPU.h
@@ -80,7 +80,7 @@ public:
 	 * @param InInputShapes The input shapes to prepare the model with.
 	 * @return 0 on success or a non-zero number otherwise.
 	 */
-	virtual int SetInputTensorShapes(TConstArrayView<FTensorShape> InInputShapes) = 0;
+	virtual int32 SetInputTensorShapes(TConstArrayView<FTensorShape> InInputShapes) = 0;
 
 	/**
 	 * Evaluate the model synchronously.
@@ -94,7 +94,7 @@ public:
 	 * @param InOutputTensors An array containing tensor bindings for each output tensor with caller owned memory big enough to contain the results on success.
 	 * @return 0 on success or a non-zero number otherwise.
 	 */
-	virtual int RunSync(TConstArrayView<FTensorBindingGPU> InInputTensors, TConstArrayView<FTensorBindingGPU> InOutputTensors) = 0;
+	virtual int32 RunSync(TConstArrayView<FTensorBindingGPU> InInputTensors, TConstArrayView<FTensorBindingGPU> InOutputTensors) = 0;
 };
 
 /**
diff --git a/Engine/Source/Runtime/NNE/Public/NNERuntimeRDG.h b/Engine/Source/Runtime/NNE/Public/NNERuntimeRDG.h
--- a/Engine/Source/Runtime/NNE/Public/NNERuntimeRDG.h
+++ b/Engine/Source/Runtime/NNE/Public/NNERuntimeRDG.h
@@ -80,7 +80,7 @@ public:
 	 * @param InInputShapes The input shapes to prepare the model with.
 	 * @return 0 on success or a non-zero number otherwise.
 	 */
-	virtual int SetInputTensorShapes(TConstArrayView<FTensorShape> InInputShapes) = 0;
+	virtual int32 SetInputTensorShapes(TConstArrayView<FTensorShape> InInputShapes) = 0;
 
 	/**
 	 * Enqueue the model graph to a FRDGBuilder.
@@ -94,7 +94,7 @@ public:
 	 * @param InOutputTensors An array containing tensor bindings for each output tensor with caller owned memory big enough to contain the results on success.
 	 * @return 0 on success or a non-zero number otherwise.
 	 */
-	virtual int EnqueueRDG(FRDGBuilder& RDGBuilder, TConstArrayView<FTensorBindingRDG> Inputs, TConstArrayView<FTensorBindingRDG> Outputs) = 0;
+	virtual int32 EnqueueRDG(FRDGBuilder& RDGBuilder, TConstArrayView<FTensorBindingRDG> Inputs, TConstArrayView<FTensorBindingRDG> Outputs) = 0;
 };
 
 /**
