NNE temptative fix for Ar.IsCountingMemory on UNNEModelData.

#rnx

diff --git a/Engine/Source/Runtime/NNE/Private/NNEModelData.cpp b/Engine/Source/Runtime/NNE/Private/NNEModelData.cpp
--- a/Engine/Source/Runtime/NNE/Private/NNEModelData.cpp
+++ b/Engine/Source/Runtime/NNE/Private/NNEModelData.cpp
@@ -187,8 +187,7 @@ void UNNEModelData::Serialize(FArchive& Ar)
 	// Store the asset version (no effect in load)
 	Ar.UsingCustomVersion(UE::NNE::ModelData::GUID);
 
-	if (Ar.IsSaving())
-	{
+	if (Ar.IsSaving() || Ar.IsCountingMemory())	{
 		bool bWriteModelData = true;
 		if (Ar.IsCooking())
 		{
@@ -289,18 +288,16 @@ void UNNEModelData::Serialize(FArchive& Ar)
 			Ar << NumItems;
 		}
 	}
-	else
+	else if (Ar.IsLoading())
 	{
-		// Read the archive
-
 		TObjectPtr<class UAssetImportData> AssetImportData;
-		int32 NumItems;
+		int32 NumItems = 0;
 		FString Name;
-		uint32 MemoryAlignment;
-		uint64 DataSize;
+		uint32 MemoryAlignment = 0;
+		uint64 DataSize = 0;
 		TArray<uint8> Data;
-		void* RawData;
-		int32 Index;
+		void* RawData = nullptr;
+		int32 Index = 0;
 
 		switch (Ar.CustomVer(UE::NNE::ModelData::GUID))
 		{
