NNE: Optimize internal data serialization

Problem:
Internal NNERuntimeFormat and NNEFormatRaw are very slow to serialize using the Archive mechanisum. Profiling reveal that TArray<uint8> is serialized elementwise slowing it down by a large amount.

Solution:
NNEFormatRaw is not meant to be serialized, thus it was converted to a regular struct to better highlight it's purpose.
NNERuntimeFormat serialization was specialized to use raw memory serialization instead of the default TArray serialization.

Test:
Win64 editor: NNE All tests green (Model test before: 670s, after:127s)
Win64 editor: NNE Smoke tests green
Win64 standalone & editor: NNE Style transfer demo.

#rnx

diff --git a/Engine/Plugins/Experimental/NNERuntimeORTCpu/Source/NNEUtils/Private/NNEUtilsModelBuilderNNE.cpp b/Engine/Plugins/Experimental/NNERuntimeORTCpu/Source/NNEUtils/Private/NNEUtilsModelBuilderNNE.cpp
--- a/Engine/Plugins/Experimental/NNERuntimeORTCpu/Source/NNEUtils/Private/NNEUtilsModelBuilderNNE.cpp
+++ b/Engine/Plugins/Experimental/NNERuntimeORTCpu/Source/NNEUtils/Private/NNEUtilsModelBuilderNNE.cpp
@@ -103,7 +103,7 @@ public:
 
 		FMemoryWriter Writer(Data);
 
-		FNNERuntimeFormat::StaticStruct()->SerializeBin(Writer, &Format);
+		Format.Serialize(Writer);
 
 		return !Data.IsEmpty();
 	}
diff --git a/Engine/Plugins/Experimental/NNERuntimeRDG/Source/NNERuntimeRDG/Private/NNERuntimeRDGBase.h b/Engine/Plugins/Experimental/NNERuntimeRDG/Source/NNERuntimeRDG/Private/NNERuntimeRDGBase.h
--- a/Engine/Plugins/Experimental/NNERuntimeRDG/Source/NNERuntimeRDG/Private/NNERuntimeRDGBase.h
+++ b/Engine/Plugins/Experimental/NNERuntimeRDG/Source/NNERuntimeRDG/Private/NNERuntimeRDGBase.h
@@ -207,7 +207,8 @@ public:
 		}
 
 		FMemoryReader Reader(InputModel.Data);
-		FNNERuntimeFormat::StaticStruct()->SerializeBin(Reader, &Format);
+
+		Format.Serialize(Reader);
 
 		TOperatorRegistryRDG<TOperatorType>* Registry = TOperatorRegistryRDG<TOperatorType>::Get();
 		check(Registry != nullptr);
diff --git a/Engine/Plugins/Experimental/NNERuntimeRDG/Source/NNERuntimeRDG/Private/NNERuntimeRDGModel.cpp b/Engine/Plugins/Experimental/NNERuntimeRDG/Source/NNERuntimeRDG/Private/NNERuntimeRDGModel.cpp
--- a/Engine/Plugins/Experimental/NNERuntimeRDG/Source/NNERuntimeRDG/Private/NNERuntimeRDGModel.cpp
+++ b/Engine/Plugins/Experimental/NNERuntimeRDG/Source/NNERuntimeRDG/Private/NNERuntimeRDGModel.cpp
@@ -19,7 +19,7 @@ bool FModelInstanceRDG::LoadModel(TConstArrayView<uint8> ModelData, FNNERuntimeF
 
 	FMemoryReaderView Reader(ModelBuffer);
 
-	FNNERuntimeFormat::StaticStruct()->SerializeBin(Reader, &Format);
+	Format.Serialize(Reader);
 
 	// Data for base class
 	InputSymbolicTensors.Empty();
diff --git a/Engine/Source/Runtime/NNE/Internal/NNERuntimeFormat.h b/Engine/Source/Runtime/NNE/Internal/NNERuntimeFormat.h
--- a/Engine/Source/Runtime/NNE/Internal/NNERuntimeFormat.h
+++ b/Engine/Source/Runtime/NNE/Internal/NNERuntimeFormat.h
@@ -29,15 +29,10 @@ enum class ENNEInferenceFormat : uint8
 	NNERT				//!< NNE Runtime format
 };
 
-USTRUCT()
 struct FNNEModelRaw
 {
-	GENERATED_USTRUCT_BODY()
-	
-	UPROPERTY(VisibleAnywhere, Category = "Neural Network Inference")
-	TArray<uint8>		Data;
+	TArray<uint8> Data;
 	
-	UPROPERTY(VisibleAnywhere, Category = "Neural Network Inference")
 	ENNEInferenceFormat	Format { ENNEInferenceFormat::Invalid };
 };
 
@@ -108,6 +103,38 @@ struct FNNERuntimeFormat
 	UPROPERTY(VisibleAnywhere, Category = "Neural Network Inference")
 	TArray<FNNEFormatOperatorDesc> Operators;
 
-	UPROPERTY(VisibleAnywhere, Category = "Neural Network Inference")
+	uint64 DataSize;
 	TArray<uint8> TensorData;
+	
+	bool Serialize(FArchive& Ar)
+	{
+		// Serialize normal UPROPERTY tagged data
+		UScriptStruct* Struct = FNNERuntimeFormat::StaticStruct();
+		Struct->SerializeTaggedProperties(Ar, reinterpret_cast<uint8*>(this), Struct, nullptr);
+
+		if (Ar.IsLoading())
+		{
+			Ar << DataSize;
+			TensorData.SetNumUninitialized(DataSize);
+			Ar.Serialize((void*)TensorData.GetData(), DataSize);
+		}
+		else
+		{
+			DataSize = TensorData.Num();
+			Ar << DataSize;
+			Ar.Serialize((void*)TensorData.GetData(), DataSize);
+		}
+
+		return true;
+	}
+
+};
+
+template<>
+struct TStructOpsTypeTraits<FNNERuntimeFormat> : public TStructOpsTypeTraitsBase2<FNNERuntimeFormat>
+{
+	enum
+	{
+		WithSerializer = true
+	};
 };
\ No newline at end of file
diff --git a/Engine/Source/Runtime/NNE/Private/NNEModelData.cpp b/Engine/Source/Runtime/NNE/Private/NNEModelData.cpp
--- a/Engine/Source/Runtime/NNE/Private/NNEModelData.cpp
+++ b/Engine/Source/Runtime/NNE/Private/NNEModelData.cpp
@@ -33,7 +33,7 @@ namespace UE::NNE::ModelData
 	const FGuid GUID(0x9513202e, 0xeba1b279, 0xf17fe5ba, 0xab90c3f2);
 	FCustomVersionRegistration NNEModelDataVersion(GUID, Version::Latest, TEXT("NNEModelDataVersion"));// Always save with the latest version
 
-	const uint32 DDCAssetVersion = 0; // Increase this value to force rebuilding cache entries
+	const uint32 DDCAssetVersion = 1; // Increase this value to force rebuilding cache entries
 
 	FString GetRuntimesAsString(TArrayView<const FString> Runtimes)
 	{
