[NNE] Assign a stable bucket to NNEModelData so it doesn't spam DDC with randomly generated buckets (one per item)

Problem:
DDC bucket are assume to be stable per system, having asset based bucket is making some stats tracking harder and creates some noise in storage.

Solution:
Use the RuntimeName as a bucket instead of the FileId from the related asset.

Note:
Fix proposed by [REDACTED_PERSON] via 28219309

Test:
NNE editor 64 : style transfer demo + nne all test
NNE standalone 64 : style transfer demo + nne smoke test
Compiled Unity.

#rnx

diff --git a/Engine/Source/Runtime/NNE/Private/NNEModelData.cpp b/Engine/Source/Runtime/NNE/Private/NNEModelData.cpp
--- a/Engine/Source/Runtime/NNE/Private/NNEModelData.cpp
+++ b/Engine/Source/Runtime/NNE/Private/NNEModelData.cpp
@@ -65,9 +65,9 @@ namespace UE::NNE::ModelData
 		return RuntimeName + "-" + FileId + "-DDCv" + FString::FromInt(DDCAssetVersion) + "-" + ModelDataIdentifier;
 	}
 
-	inline UE::DerivedData::FCacheKey CreateCacheKey(const FString& FileId, const FString& RequestId)
+	inline UE::DerivedData::FCacheKey CreateCacheKey(const FString& RequestId)
 	{
-		return { UE::DerivedData::FCacheBucket(FWideStringView(*FileId)), FIoHash::HashBuffer(MakeMemoryView(FTCHARToUTF8(RequestId))) };
+		return { UE::DerivedData::FCacheBucket(TEXT("NNEModelData")), FIoHash::HashBuffer(MakeMemoryView(FTCHARToUTF8(RequestId))) };
 	}
 
 	inline void PutIntoDDC(const FGuid& FileId, const FString& RuntimeName, const FString& ModelDataIdentifier, TSharedPtr<UE::NNE::FSharedModelData> Data)
@@ -82,7 +82,7 @@ namespace UE::NNE::ModelData
 		Requests.SetNum(1);
 
 		Requests[0].Name = FString("Put-") + RequestId;
-		Requests[0].Key = CreateCacheKey(FileIdStr, RequestId);
+		Requests[0].Key = CreateCacheKey(RequestId);
 		Requests[0].Value = UE::DerivedData::FValue::Compress(FCompositeBuffer(MakeSharedBufferFromArray(TArray<uint32>({ Data->GetMemoryAlignment() })), FSharedBuffer::MakeView(Data->GetView().GetData(), Data->GetView().Num())));
 
 		UE::DerivedData::FRequestOwner BlockingPutOwner(UE::DerivedData::EPriority::Blocking);
@@ -99,7 +99,7 @@ namespace UE::NNE::ModelData
 		Requests.SetNum(1);
 
 		Requests[0].Name = FString("Get-") + RequestId;
-		Requests[0].Key = CreateCacheKey(FileIdStr, RequestId);
+		Requests[0].Key = CreateCacheKey(RequestId);
 
 		TSharedPtr<UE::NNE::FSharedModelData> Result;
 		UE::DerivedData::FRequestOwner BlockingGetOwner(UE::DerivedData::EPriority::Blocking);
