[NNE] QA Fix for standalone tests: Reenable all cooked runtime.

Problem:
* UNNEModelData is serialized as an empty list for cooked asset, meaning standalone test were not running against any runtimes!

Solution:
* Clarify UNNEModelData contract around TargetRuntime and handle the editor vs cooked asset case in the NNE QA framework.

Tests:
* Editor win64: All test + style transfer demo
* Editor standaonle win 64: All test (green but known oidn ones, to be addressed separately) + style transfer demo

#rnx

diff --git a/Engine/Source/Runtime/NNE/Public/NNEModelData.h b/Engine/Source/Runtime/NNE/Public/NNEModelData.h
--- a/Engine/Source/Runtime/NNE/Public/NNEModelData.h
+++ b/Engine/Source/Runtime/NNE/Public/NNEModelData.h
@@ -91,7 +91,8 @@ public:
 	void Init(const FString& Type, TConstArrayView<uint8> Buffer);
 
 	/**
-	 * Get the target runtimes this model data will be cooked for. An empty list means all runtimes.
+	 * In editor: Get the target runtimes this model data will be cooked for. An empty list means all runtimes.
+	 * In standalone: An empty list.
 	 *
 	 * @return The target runtimes names.
 	 */
@@ -107,7 +108,8 @@ public:
 	/**
 	 * Get the type of data inside FileData.
 	 *
-	 * The FileType identifies the type of data inside FileData and typically is the extension of the file used to create the asset.
+	 * In editor: The FileType identifies the type of data inside FileData and typically is the extension of the file used to create the asset.
+	 * In standalone: An empty string.
 	 *
 	 * @return The FileType.
 	 */
@@ -116,7 +118,8 @@ public:
 	/**
 	 * Get read only access to FileData.
 	 *
-	 * The FileData contains the binary data of the file which has been used to create the asset.
+	 * In editor: The FileData contains the binary data of the file which has been used to create the asset.
+	 * In standalone: An empty array.
 	 *
 	 * @return The FileData.
 	 */
