[NNE] Allow ability to use DDC for cooking for dev iteration

Problem:
* UNNEModelData is using onnx for model optimisation itself slow to close the session after model optimisation, this combined with the fact that we always reoptimized from scratch when cooking create a large bottleneck when running our tests suits.
* NNE use the DDC in the editor flow to fetch optimized model, however there might be situation where the DDC is corrupted and there is no easy way to get out of this situation.

Solution:
* Add project level editor configuration back by CVars for NNE, allow to enable the DDC for cooking or to disable DDC completely.

Tests:
* Editor win64: All test + style transfer demo
* Editor standalone win 64: All test + style transfer demo

#rnx

diff --git a/Engine/Source/Editor/NNEEditor/NNEEditor.Build.cs b/Engine/Source/Editor/NNEEditor/NNEEditor.Build.cs
--- a/Engine/Source/Editor/NNEEditor/NNEEditor.Build.cs
+++ b/Engine/Source/Editor/NNEEditor/NNEEditor.Build.cs
@@ -18,6 +18,7 @@ public class NNEEditor : ModuleRules
 		PrivateDependencyModuleNames.AddRange(
 			new string[]
 			{
+				"DeveloperSettings",
 				"Slate",
 				"SlateCore",
 				"UnrealEd",
diff --git a/Engine/Source/Editor/NNEEditor/Private/NNEEditorSettings.cpp b/Engine/Source/Editor/NNEEditor/Private/NNEEditorSettings.cpp
new file mode 100644
--- /dev/null
+++ b/Engine/Source/Editor/NNEEditor/Private/NNEEditorSettings.cpp
@@ -0,0 +1,17 @@
+// Copyright Epic Games, Inc. All Rights Reserved.
+
+#include "NNEEditorSettings.h"
+
+#include "HAL/IConsoleManager.h"
+#include "Misc/ConfigUtilities.h"
+
+void UNNEEditorSettings::PostInitProperties()
+{
+	if (IsTemplate())
+	{
+		// We want the .ini file to have precedence over the CVar constructor, so we apply the ini to the CVar before following the regular UDeveloperSettingsBackedByCVars flow
+		UE::ConfigUtilities::ApplyCVarSettingsFromIni(TEXT("/Script/NNEEditor.NNEEditorSettings"), *GEditorIni, ECVF_SetByProjectSetting);
+	}
+
+	Super::PostInitProperties();
+}
\ No newline at end of file
diff --git a/Engine/Source/Editor/NNEEditor/Private/NNEEditorSettings.h b/Engine/Source/Editor/NNEEditor/Private/NNEEditorSettings.h
new file mode 100644
--- /dev/null
+++ b/Engine/Source/Editor/NNEEditor/Private/NNEEditorSettings.h
@@ -0,0 +1,33 @@
+// Copyright Epic Games, Inc. All Rights Reserved.
+
+#pragma once
+
+#include "CoreMinimal.h"
+#include "Engine/DeveloperSettings.h"
+#include "Engine/DeveloperSettingsBackedByCVars.h"
+#include "Misc/ConfigUtilities.h"
+
+#include "NNEEditorSettings.generated.h"
+
+UCLASS(config=Editor, defaultconfig, meta = (DisplayName = "Neural Network Engine (NNE)"))
+class UNNEEditorSettings : public UDeveloperSettingsBackedByCVars
+{
+	GENERATED_BODY()
+
+public:
+
+	virtual FName GetContainerName() const override { return FName("Project"); }
+	virtual FName GetCategoryName() const override { return FName("Editor"); }
+	virtual void PostInitProperties() override;
+	
+	UPROPERTY(config, EditAnywhere, BlueprintReadWrite, Category = Editor, meta = (
+		ConsoleVariable = "nne.Editor.UseDDC", DisplayName = "Use DDC for Neural Network optimizations",
+		ToolTip = "True if NNE should use the DDC cache to fetch previously optimized Neural Network."))
+	bool bUseDDCCache;
+
+	UPROPERTY(config, EditAnywhere, BlueprintReadWrite, Category = Editor, meta = (
+		EditCondition = "bUseDDCCache",
+		ConsoleVariable = "nne.Editor.UseDDCForCooking", DisplayName = "Use DDC for Neural Network optimizations while cooking",
+		ToolTip = "True if NNE should use the DDC cache to fetch previously optimized Neural Network while cooking."))
+	bool bUseDDCCacheForCooking;
+};
\ No newline at end of file
diff --git a/Engine/Source/Runtime/NNE/Private/NNEModelData.cpp b/Engine/Source/Runtime/NNE/Private/NNEModelData.cpp
--- a/Engine/Source/Runtime/NNE/Private/NNEModelData.cpp
+++ b/Engine/Source/Runtime/NNE/Private/NNEModelData.cpp
@@ -15,6 +15,7 @@
 #include "DerivedDataCache.h"
 #include "DerivedDataCacheKey.h"
 #include "DerivedDataRequestOwner.h"
+#include "HAL/IConsoleManager.h"
 #include "Memory/CompositeBuffer.h"
 #endif // WITH_EDITOR
 
@@ -174,6 +175,24 @@ namespace UE::NNE
 	{
 		return MemoryAlignment;
 	}
+
+#if WITH_EDITOR
+	namespace ConsoleVariables
+	{
+		static TAutoConsoleVariable<bool> CVarNNEEditorUseDDC(
+			TEXT("nne.Editor.UseDDC"),
+			true,
+			TEXT("Indicates whether NNE wether NNE should use the DDC cache to speed up SharedModelData creation in the editor, see also nne.Editor.UseDDCForCooking."),
+			ECVF_Default);
+
+		static TAutoConsoleVariable<bool> CVarNNEEditorUseDDCForCooking(
+			TEXT("nne.Editor.UseDDCForCooking"),
+			false,
+			TEXT("Indicates whether NNE wether NNE should use the DDC cache to speed up SharedModelData creation while cooking, if nne.Editor.UseDDC is false cooking won't use the DDC even if nne.Editor.UseDDCForCooking is true."),
+			ECVF_Default);
+	} // ConsoleVariables
+#endif //WITH_EDITOR
+
 } // UE::NNE
 
 void UNNEModelData::GetAssetRegistryTags(TArray<FAssetRegistryTag>& OutTags) const
@@ -221,30 +240,45 @@ void UNNEModelData::Serialize(FArchive& Ar)
 
 				for (const FString& RuntimeName : CookRuntimeNames)
 				{
-					TSharedPtr<UE::NNE::FSharedModelData> CreatedData = UE::NNE::ModelData::CreateModelData(RuntimeName, FileType, FileData, FileId, Ar.GetArchiveState().CookingTarget());
-					if (CreatedData.IsValid() && CreatedData->GetView().Num() > 0)
-					{
-						ModelData.Add(RuntimeName, CreatedData);
+					TSharedPtr<UE::NNE::FSharedModelData> SharedModelData;
 #if WITH_EDITOR
-						TWeakInterfacePtr<INNERuntime> NNERuntime = UE::NNE::GetRuntime<INNERuntime>(RuntimeName);
-						if (NNERuntime.IsValid())
+					TWeakInterfacePtr<INNERuntime> NNERuntime = UE::NNE::GetRuntime<INNERuntime>(RuntimeName);
+					FString ModelDataIdentifier;
+					if (NNERuntime.IsValid())
+					{
+						ModelDataIdentifier = NNERuntime->GetModelDataIdentifier(FileType, FileData, FileId, Ar.GetArchiveState().CookingTarget());
+						if (ModelDataIdentifier.Len() > 0)
 						{
-							FString ModelDataIdentifier = NNERuntime->GetModelDataIdentifier(FileType, FileData, FileId, Ar.GetArchiveState().CookingTarget());
-							if (ModelDataIdentifier.Len() > 0)
-							{
-								UE::NNE::ModelData::PutIntoDDC(FileId, RuntimeName, ModelDataIdentifier, CreatedData);
-							}
-							else
+							if (UE::NNE::ConsoleVariables::CVarNNEEditorUseDDC.GetValueOnGameThread() &&
+								UE::NNE::ConsoleVariables::CVarNNEEditorUseDDCForCooking.GetValueOnGameThread())
 							{
-								UE_LOG(LogNNE, Warning, TEXT("UNNEModelData: Runtime '%s' returned an empty string as a ModelDataIdentifier while cooking. GetModelDataIdentifier should always return a valid identifier."), *RuntimeName);
+								SharedModelData = UE::NNE::ModelData::GetFromDDC(FileId, RuntimeName, ModelDataIdentifier);
 							}
 						}
 						else
 						{
-							UE_LOG(LogNNE, Warning, TEXT("UNNEModelData: Runtime '%s' is among the cooked runtimes but instance is invalid."), *RuntimeName);
+							UE_LOG(LogNNE, Warning, TEXT("UNNEModelData: Runtime '%s' returned an empty string as a ModelDataIdentifier while cooking. GetModelDataIdentifier should always return a valid identifier."), *RuntimeName);
+						}
+					}
+					else
+					{
+						UE_LOG(LogNNE, Warning, TEXT("UNNEModelData: Runtime '%s' is among the cooked runtimes but instance is invalid."), *RuntimeName);
+					}
+#endif //WITH_EDITOR
+					if (!SharedModelData.IsValid() || SharedModelData->GetView().Num() == 0)
+					{
+						SharedModelData = UE::NNE::ModelData::CreateModelData(RuntimeName, FileType, FileData, FileId, Ar.GetArchiveState().CookingTarget());
+#if WITH_EDITOR
+						if (SharedModelData.IsValid() && (SharedModelData->GetView().Num() > 0) && NNERuntime.IsValid() && (ModelDataIdentifier.Len() > 0))
+						{
+							UE::NNE::ModelData::PutIntoDDC(FileId, RuntimeName, ModelDataIdentifier, SharedModelData);
 						}
 #endif //WITH_EDITOR
 					}
+					if (SharedModelData.IsValid() && SharedModelData->GetView().Num() > 0)
+					{
+						ModelData.Add(RuntimeName, SharedModelData);
+					}
 				}
 			}
 		}
@@ -457,12 +491,15 @@ TSharedPtr<UE::NNE::FSharedModelData> UNNEModelData::GetModelData(const FString&
 		return TSharedPtr<UE::NNE::FSharedModelData>();
 	}
 
-	// Check if we have a DDC cache hit
-	TSharedPtr<UE::NNE::FSharedModelData> RemoteData = UE::NNE::ModelData::GetFromDDC(FileId, RuntimeName, ModelDataIdentifier);
-	if (RemoteData.IsValid() && RemoteData->GetView().Num() > 0)
+	if (UE::NNE::ConsoleVariables::CVarNNEEditorUseDDC.GetValueOnGameThread())
 	{
-		ModelData.Add(RuntimeName, RemoteData);
-		return RemoteData;
+		// Check if we have a DDC cache hit
+		TSharedPtr<UE::NNE::FSharedModelData> RemoteData = UE::NNE::ModelData::GetFromDDC(FileId, RuntimeName, ModelDataIdentifier);
+		if (RemoteData.IsValid() && RemoteData->GetView().Num() > 0)
+		{
+			ModelData.Add(RuntimeName, RemoteData);
+			return RemoteData;
+		}
 	}
 #endif //WITH_EDITOR
 
@@ -487,4 +524,4 @@ TSharedPtr<UE::NNE::FSharedModelData> UNNEModelData::GetModelData(const FString&
 void UNNEModelData::ClearModelData()
 {
 	ModelData.Empty();
-}
+}
\ No newline at end of file
