[NNE] Adding AdditionalFileData

Problem: In some cases, importing a neural netowrk file requires additional data to be stored besides the original file. E.g. onnx based runtimes may require a sidecare file storing all the weights (required for models > 2GB) or mlir based runtimes want to parse and store function meta data.

Solution: Add the functionality to add AdditionalFileData to UNNEModelData. This data is stored in a Map<FString, TArray<uint8>> next to file data and passed to model data creation functions during cooking. It will not be contained in cooked assets.

Tested NNEIncubator and IREEPlugin: non-unity compilation, local win64 PIE and packaged + automation tests, complete engine presubmit preflight

#rnx

diff --git a/Engine/Plugins/Experimental/NNERuntimeBasicCpu/Source/NNERuntimeBasicCpu/Private/NNERuntimeBasicCpu.cpp b/Engine/Plugins/Experimental/NNERuntimeBasicCpu/Source/NNERuntimeBasicCpu/Private/NNERuntimeBasicCpu.cpp
--- a/Engine/Plugins/Experimental/NNERuntimeBasicCpu/Source/NNERuntimeBasicCpu/Private/NNERuntimeBasicCpu.cpp
+++ b/Engine/Plugins/Experimental/NNERuntimeBasicCpu/Source/NNERuntimeBasicCpu/Private/NNERuntimeBasicCpu.cpp
@@ -9,7 +9,7 @@
 // various bits of data.
 const uint32 UNNERuntimeBasicCpuImpl::Alignment = 64;
 
-bool UNNERuntimeBasicCpuImpl::CanCreateModelData(FString FileType, TConstArrayView<uint8> FileData, FGuid FileId, const ITargetPlatform* TargetPlatform) const
+bool UNNERuntimeBasicCpuImpl::CanCreateModelData(const FString& FileType, TConstArrayView<uint8> FileData, const TMap<FString, TConstArrayView<uint8>>& AdditionalFileData, const FGuid& FileId, const ITargetPlatform* TargetPlatform) const
 {
 	if (FileType.Compare("ubnne", ESearchCase::IgnoreCase) != 0)
 	{
@@ -39,9 +39,9 @@ bool UNNERuntimeBasicCpuImpl::CanCreateModelData(FString FileType, TConstArrayVi
 	return true;
 }
 
-TSharedPtr<UE::NNE::FSharedModelData> UNNERuntimeBasicCpuImpl::CreateModelData(FString FileType, TConstArrayView<uint8> FileData, FGuid FileId, const ITargetPlatform* TargetPlatform)
+TSharedPtr<UE::NNE::FSharedModelData> UNNERuntimeBasicCpuImpl::CreateModelData(const FString& FileType, TConstArrayView<uint8> FileData, const TMap<FString, TConstArrayView<uint8>>& AdditionalFileData, const FGuid& FileId, const ITargetPlatform* TargetPlatform)
 {
-	if (!CanCreateModelData(FileType, FileData, FileId, TargetPlatform))
+	if (!CanCreateModelData(FileType, FileData, AdditionalFileData, FileId, TargetPlatform))
 	{
 		return nullptr;
 	}
@@ -56,7 +56,7 @@ TSharedPtr<UE::NNE::FSharedModelData> UNNERuntimeBasicCpuImpl::CreateModelData(F
 	return MakeShared<UE::NNE::FSharedModelData>(FSharedBuffer::TakeOwnership(ModelData, FileData.Num(), FMemory::Free), Alignment);
 }
 
-FString UNNERuntimeBasicCpuImpl::GetModelDataIdentifier(FString FileType, TConstArrayView<uint8> FileData, FGuid FileId, const ITargetPlatform* TargetPlatform)
+FString UNNERuntimeBasicCpuImpl::GetModelDataIdentifier(const FString& FileType, TConstArrayView<uint8> FileData, const TMap<FString, TConstArrayView<uint8>>& AdditionalFileData, const FGuid& FileId, const ITargetPlatform* TargetPlatform) const
 {
 	return FileId.ToString(EGuidFormats::Digits) + "-" + FString::FromInt(UE::NNE::RuntimeBasic::FModelCPU::ModelMagicNumber);
 }
diff --git a/Engine/Plugins/Experimental/NNERuntimeBasicCpu/Source/NNERuntimeBasicCpu/Private/NNERuntimeBasicCpu.h b/Engine/Plugins/Experimental/NNERuntimeBasicCpu/Source/NNERuntimeBasicCpu/Private/NNERuntimeBasicCpu.h
--- a/Engine/Plugins/Experimental/NNERuntimeBasicCpu/Source/NNERuntimeBasicCpu/Private/NNERuntimeBasicCpu.h
+++ b/Engine/Plugins/Experimental/NNERuntimeBasicCpu/Source/NNERuntimeBasicCpu/Private/NNERuntimeBasicCpu.h
@@ -25,9 +25,9 @@ public:
 
 	virtual FString GetRuntimeName() const override { return TEXT("NNERuntimeBasicCpu"); };
 
-	virtual bool CanCreateModelData(FString FileType, TConstArrayView<uint8> FileData, FGuid FileId, const ITargetPlatform* TargetPlatform) const override;
-	virtual TSharedPtr<UE::NNE::FSharedModelData> CreateModelData(FString FileType, TConstArrayView<uint8> FileData, FGuid FileId, const ITargetPlatform* TargetPlatform) override;
-	virtual FString GetModelDataIdentifier(FString FileType, TConstArrayView<uint8> FileData, FGuid FileId, const ITargetPlatform* TargetPlatform) override;
+	virtual bool CanCreateModelData(const FString& FileType, TConstArrayView<uint8> FileData, const TMap<FString, TConstArrayView<uint8>>& AdditionalFileData, const FGuid& FileId, const ITargetPlatform* TargetPlatform) const override;
+	virtual TSharedPtr<UE::NNE::FSharedModelData> CreateModelData(const FString& FileType, TConstArrayView<uint8> FileData, const TMap<FString, TConstArrayView<uint8>>& AdditionalFileData, const FGuid& FileId, const ITargetPlatform* TargetPlatform) override;
+	virtual FString GetModelDataIdentifier(const FString& FileType, TConstArrayView<uint8> FileData, const TMap<FString, TConstArrayView<uint8>>& AdditionalFileData, const FGuid& FileId, const ITargetPlatform* TargetPlatform) const override;
 
 	virtual bool CanCreateModelCPU(TObjectPtr<UNNEModelData> ModelData) const override;
 	virtual TSharedPtr<UE::NNE::IModelCPU> CreateModelCPU(TObjectPtr<UNNEModelData> ModelData) override;
diff --git a/Engine/Plugins/Experimental/NNERuntimeIREE/Source/NNERuntimeIREE/Private/NNERuntimeIREECpu.cpp b/Engine/Plugins/Experimental/NNERuntimeIREE/Source/NNERuntimeIREE/Private/NNERuntimeIREECpu.cpp
--- a/Engine/Plugins/Experimental/NNERuntimeIREE/Source/NNERuntimeIREE/Private/NNERuntimeIREECpu.cpp
+++ b/Engine/Plugins/Experimental/NNERuntimeIREE/Source/NNERuntimeIREE/Private/NNERuntimeIREECpu.cpp
@@ -109,7 +109,7 @@ FString UNNERuntimeIREECpu::GetRuntimeName() const
 	return TEXT("NNERuntimeIREECpu");
 }
 
-bool UNNERuntimeIREECpu::CanCreateModelData(FString FileType, TConstArrayView<uint8> FileData, FGuid FileId, const ITargetPlatform* TargetPlatform) const
+bool UNNERuntimeIREECpu::CanCreateModelData(const FString& FileType, TConstArrayView<uint8> FileData, const TMap<FString, TConstArrayView<uint8>>& AdditionalFileData, const FGuid& FileId, const ITargetPlatform* TargetPlatform) const
 {
 #if WITH_EDITOR
 	return 	FileType.Compare(TEXT("mlir"), ESearchCase::IgnoreCase) == 0;
@@ -118,13 +118,13 @@ bool UNNERuntimeIREECpu::CanCreateModelData(FString FileType, TConstArrayView<ui
 #endif // WITH_EDITOR
 }
 
-TSharedPtr<UE::NNE::FSharedModelData> UNNERuntimeIREECpu::CreateModelData(FString FileType, TConstArrayView<uint8> FileData, FGuid FileId, const ITargetPlatform* TargetPlatform)
+TSharedPtr<UE::NNE::FSharedModelData> UNNERuntimeIREECpu::CreateModelData(const FString& FileType, TConstArrayView<uint8> FileData, const TMap<FString, TConstArrayView<uint8>>& AdditionalFileData, const FGuid& FileId, const ITargetPlatform* TargetPlatform)
 {
 #if WITH_EDITOR
 	using namespace UE::NNERuntimeIREECpu::Private;
 
 	FString TargetPlatformName = TargetPlatform ? TargetPlatform->IniPlatformName() : UGameplayStatics::GetPlatformName();
-	if (!CanCreateModelData(FileType, FileData, FileId, TargetPlatform))
+	if (!CanCreateModelData(FileType, FileData, AdditionalFileData, FileId, TargetPlatform))
 	{
 		UE_LOG(LogTemp, Warning, TEXT("UNNERuntimeIREECpu cannot create the model data with id %s (Filetype: %s) for platform %s"), *FileId.ToString(EGuidFormats::Digits).ToLower(), *FileType, *TargetPlatformName);
 		return TSharedPtr<UE::NNE::FSharedModelData>();
@@ -194,7 +194,8 @@ TSharedPtr<UE::NNE::FSharedModelData> UNNERuntimeIREECpu::CreateModelData(FStrin
 	FMemoryWriter Writer(ResultData);
 	Writer << UNNERuntimeIREECpu::GUID;
 	Writer << UNNERuntimeIREECpu::Version;
-	Writer << FileId;
+	FGuid FileIdCopy = FileId;
+	Writer << FileIdCopy;
 
 	Writer << ModuleDataString;
 
@@ -215,7 +216,7 @@ TSharedPtr<UE::NNE::FSharedModelData> UNNERuntimeIREECpu::CreateModelData(FStrin
 #endif // WITH_EDITOR
 }
 
-FString UNNERuntimeIREECpu::GetModelDataIdentifier(FString FileType, TConstArrayView<uint8> FileData, FGuid FileId, const ITargetPlatform* TargetPlatform)
+FString UNNERuntimeIREECpu::GetModelDataIdentifier(const FString& FileType, TConstArrayView<uint8> FileData, const TMap<FString, TConstArrayView<uint8>>& AdditionalFileData, const FGuid& FileId, const ITargetPlatform* TargetPlatform) const
 {
 	// Leave architecture blank as there is only one model data for all architectures of a given platform
 	FString PlatformName = TargetPlatform ? TargetPlatform->IniPlatformName() : UGameplayStatics::GetPlatformName();
@@ -407,9 +408,9 @@ void UNNERuntimeIREECpu::GetUpdatedPlatformConfig(const FString& PlatformName, F
 
 FString UNNERuntimeIREECpu::GetRuntimeName() const { return ""; };
 
-bool UNNERuntimeIREECpu::CanCreateModelData(FString FileType, TConstArrayView<uint8> FileData, FGuid FileId, const ITargetPlatform* TargetPlatform) const { return false; };
-TSharedPtr<UE::NNE::FSharedModelData> UNNERuntimeIREECpu::CreateModelData(FString FileType, TConstArrayView<uint8> FileData, FGuid FileId, const ITargetPlatform* TargetPlatform) { return TSharedPtr<UE::NNE::FSharedModelData>(); };
-FString UNNERuntimeIREECpu::GetModelDataIdentifier(FString FileType, TConstArrayView<uint8> FileData, FGuid FileId, const ITargetPlatform* TargetPlatform) { return ""; };
+bool UNNERuntimeIREECpu::CanCreateModelData(const FString& FileType, TConstArrayView<uint8> FileData, const TMap<FString, TConstArrayView<uint8>>& AdditionalFileData, const FGuid& FileId, const ITargetPlatform* TargetPlatform) const { return false; };
+TSharedPtr<UE::NNE::FSharedModelData> UNNERuntimeIREECpu::CreateModelData(const FString& FileType, TConstArrayView<uint8> FileData, const TMap<FString, TConstArrayView<uint8>>& AdditionalFileData, const FGuid& FileId, const ITargetPlatform* TargetPlatform) { return TSharedPtr<UE::NNE::FSharedModelData>(); };
+FString UNNERuntimeIREECpu::GetModelDataIdentifier(const FString& FileType, TConstArrayView<uint8> FileData, const TMap<FString, TConstArrayView<uint8>>& AdditionalFileData, const FGuid& FileId, const ITargetPlatform* TargetPlatform) const { return ""; };
 
 bool UNNERuntimeIREECpu::CanCreateModelCPU(TObjectPtr<UNNEModelData> ModelData) const { return false; };
 TSharedPtr<UE::NNE::IModelCPU> UNNERuntimeIREECpu::CreateModelCPU(TObjectPtr<UNNEModelData> ModelData) { return TSharedPtr<UE::NNE::IModelCPU>(); };
diff --git a/Engine/Plugins/Experimental/NNERuntimeIREE/Source/NNERuntimeIREE/Private/NNERuntimeIREECpu.h b/Engine/Plugins/Experimental/NNERuntimeIREE/Source/NNERuntimeIREE/Private/NNERuntimeIREECpu.h
--- a/Engine/Plugins/Experimental/NNERuntimeIREE/Source/NNERuntimeIREE/Private/NNERuntimeIREECpu.h
+++ b/Engine/Plugins/Experimental/NNERuntimeIREE/Source/NNERuntimeIREE/Private/NNERuntimeIREECpu.h
@@ -25,9 +25,9 @@ public:
 
 	//~ Begin INNERuntime Interface
 	virtual FString GetRuntimeName() const override;
-	virtual bool CanCreateModelData(FString FileType, TConstArrayView<uint8> FileData, FGuid FileId, const ITargetPlatform* TargetPlatform) const override;
-	virtual TSharedPtr<UE::NNE::FSharedModelData> CreateModelData(FString FileType, TConstArrayView<uint8> FileData, FGuid FileId, const ITargetPlatform* TargetPlatform) override;
-	virtual FString GetModelDataIdentifier(FString FileType, TConstArrayView<uint8> FileData, FGuid FileId, const ITargetPlatform* TargetPlatform) override;
+	virtual bool CanCreateModelData(const FString& FileType, TConstArrayView<uint8> FileData, const TMap<FString, TConstArrayView<uint8>>& AdditionalFileData, const FGuid& FileId, const ITargetPlatform* TargetPlatform) const override;
+	virtual TSharedPtr<UE::NNE::FSharedModelData> CreateModelData(const FString& FileType, TConstArrayView<uint8> FileData, const TMap<FString, TConstArrayView<uint8>>& AdditionalFileData, const FGuid& FileId, const ITargetPlatform* TargetPlatform) override;
+	virtual FString GetModelDataIdentifier(const FString& FileType, TConstArrayView<uint8> FileData, const TMap<FString, TConstArrayView<uint8>>& AdditionalFileData, const FGuid& FileId, const ITargetPlatform* TargetPlatform) const override;
 	//~ End INNERuntime Interface
 
 	//~ Begin INNERuntimeCPU Interface
diff --git a/Engine/Plugins/Experimental/NNERuntimeIREE/Source/NNERuntimeIREE/Private/NNERuntimeIREEGpu.cpp b/Engine/Plugins/Experimental/NNERuntimeIREE/Source/NNERuntimeIREE/Private/NNERuntimeIREEGpu.cpp
--- a/Engine/Plugins/Experimental/NNERuntimeIREE/Source/NNERuntimeIREE/Private/NNERuntimeIREEGpu.cpp
+++ b/Engine/Plugins/Experimental/NNERuntimeIREE/Source/NNERuntimeIREE/Private/NNERuntimeIREEGpu.cpp
@@ -30,7 +30,7 @@ FString UNNERuntimeIREEVulkan::GetRuntimeName() const
 	return TEXT("NNERuntimeIREEVulkan");
 }
 
-bool UNNERuntimeIREEGpu::CanCreateModelData(FString FileType, TConstArrayView<uint8> FileData, FGuid FileId, const ITargetPlatform* TargetPlatform) const
+bool UNNERuntimeIREEGpu::CanCreateModelData(const FString& FileType, TConstArrayView<uint8> FileData, const TMap<FString, TConstArrayView<uint8>>& AdditionalFileData, const FGuid& FileId, const ITargetPlatform* TargetPlatform) const
 {
 #if WITH_EDITOR
 	return	FileType.Compare(TEXT("mlir"), ESearchCase::IgnoreCase) == 0;
@@ -39,12 +39,12 @@ bool UNNERuntimeIREEGpu::CanCreateModelData(FString FileType, TConstArrayView<ui
 #endif // WITH_EDITOR
 }
 
-TSharedPtr<UE::NNE::FSharedModelData> UNNERuntimeIREEGpu::CreateModelData(FString FileType, TConstArrayView<uint8> FileData, FGuid FileId, const ITargetPlatform* TargetPlatform)
+TSharedPtr<UE::NNE::FSharedModelData> UNNERuntimeIREEGpu::CreateModelData(const FString& FileType, TConstArrayView<uint8> FileData, const TMap<FString, TConstArrayView<uint8>>& AdditionalFileData, const FGuid& FileId, const ITargetPlatform* TargetPlatform)
 {
 	return TSharedPtr<UE::NNE::FSharedModelData>();
 }
 
-FString UNNERuntimeIREEGpu::GetModelDataIdentifier(FString FileType, TConstArrayView<uint8> FileData, FGuid FileId, const ITargetPlatform* TargetPlatform)
+FString UNNERuntimeIREEGpu::GetModelDataIdentifier(const FString& FileType, TConstArrayView<uint8> FileData, const TMap<FString, TConstArrayView<uint8>>& AdditionalFileData, const FGuid& FileId, const ITargetPlatform* TargetPlatform) const
 {
 	FString PlatformName = TargetPlatform ? TargetPlatform->IniPlatformName() : UGameplayStatics::GetPlatformName();
 	return GetRuntimeName() + "-" + GetGUID().ToString(EGuidFormats::Digits) + "-" + FString::FromInt(GetVersion()) + "-" + FileId.ToString(EGuidFormats::Digits) + "-" + PlatformName;
@@ -153,9 +153,9 @@ FString UNNERuntimeIREEGpu::GetRuntimeName() const { return TEXT(""); }
 FString UNNERuntimeIREECuda::GetRuntimeName() const { return TEXT(""); }
 FString UNNERuntimeIREEVulkan::GetRuntimeName() const { return TEXT(""); }
 
-bool UNNERuntimeIREEGpu::CanCreateModelData(FString FileType, TConstArrayView<uint8> FileData, FGuid FileId, const ITargetPlatform* TargetPlatform) const { return false; };
-TSharedPtr<UE::NNE::FSharedModelData> UNNERuntimeIREEGpu::CreateModelData(FString FileType, TConstArrayView<uint8> FileData, FGuid FileId, const ITargetPlatform* TargetPlatform) { return TSharedPtr<UE::NNE::FSharedModelData>(); };
-FString UNNERuntimeIREEGpu::GetModelDataIdentifier(FString FileType, TConstArrayView<uint8> FileData, FGuid FileId, const ITargetPlatform* TargetPlatform) { return ""; };
+bool UNNERuntimeIREEGpu::CanCreateModelData(const FString& FileType, TConstArrayView<uint8> FileData, const TMap<FString, TConstArrayView<uint8>>& AdditionalFileData, const FGuid& FileId, const ITargetPlatform* TargetPlatform) const { return false; };
+TSharedPtr<UE::NNE::FSharedModelData> UNNERuntimeIREEGpu::CreateModelData(const FString& FileType, TConstArrayView<uint8> FileData, const TMap<FString, TConstArrayView<uint8>>& AdditionalFileData, const FGuid& FileId, const ITargetPlatform* TargetPlatform) { return TSharedPtr<UE::NNE::FSharedModelData>(); };
+FString UNNERuntimeIREEGpu::GetModelDataIdentifier(const FString& FileType, TConstArrayView<uint8> FileData, const TMap<FString, TConstArrayView<uint8>>& AdditionalFileData, const FGuid& FileId, const ITargetPlatform* TargetPlatform) const { return ""; };
 
 bool UNNERuntimeIREEGpu::CanCreateModelGPU(TObjectPtr<UNNEModelData> ModelData) const { return false; };
 TSharedPtr<UE::NNE::IModelGPU> UNNERuntimeIREEGpu::CreateModelGPU(TObjectPtr<UNNEModelData> ModelData) { return TSharedPtr<UE::NNE::IModelGPU>(); };
diff --git a/Engine/Plugins/Experimental/NNERuntimeIREE/Source/NNERuntimeIREE/Private/NNERuntimeIREEGpu.h b/Engine/Plugins/Experimental/NNERuntimeIREE/Source/NNERuntimeIREE/Private/NNERuntimeIREEGpu.h
--- a/Engine/Plugins/Experimental/NNERuntimeIREE/Source/NNERuntimeIREE/Private/NNERuntimeIREEGpu.h
+++ b/Engine/Plugins/Experimental/NNERuntimeIREE/Source/NNERuntimeIREE/Private/NNERuntimeIREEGpu.h
@@ -21,9 +21,9 @@ public:
 
 	//~ Begin INNERuntime Interface
 	virtual FString GetRuntimeName() const override;
-	virtual bool CanCreateModelData(FString FileType, TConstArrayView<uint8> FileData, FGuid FileId, const ITargetPlatform* TargetPlatform) const override;
-	virtual TSharedPtr<UE::NNE::FSharedModelData> CreateModelData(FString FileType, TConstArrayView<uint8> FileData, FGuid FileId, const ITargetPlatform* TargetPlatform) override;
-	virtual FString GetModelDataIdentifier(FString FileType, TConstArrayView<uint8> FileData, FGuid FileId, const ITargetPlatform* TargetPlatform) override;
+	virtual bool CanCreateModelData(const FString& FileType, TConstArrayView<uint8> FileData, const TMap<FString, TConstArrayView<uint8>>& AdditionalFileData, const FGuid& FileId, const ITargetPlatform* TargetPlatform) const override;
+	virtual TSharedPtr<UE::NNE::FSharedModelData> CreateModelData(const FString& FileType, TConstArrayView<uint8> FileData, const TMap<FString, TConstArrayView<uint8>>& AdditionalFileData, const FGuid& FileId, const ITargetPlatform* TargetPlatform) override;
+	virtual FString GetModelDataIdentifier(const FString& FileType, TConstArrayView<uint8> FileData, const TMap<FString, TConstArrayView<uint8>>& AdditionalFileData, const FGuid& FileId, const ITargetPlatform* TargetPlatform) const override;
 	//~ End INNERuntime Interface
 
 	//~ Begin INNERuntimeGPU Interface
diff --git a/Engine/Plugins/Experimental/NNERuntimeIREE/Source/NNERuntimeIREE/Private/NNERuntimeIREERdg.cpp b/Engine/Plugins/Experimental/NNERuntimeIREE/Source/NNERuntimeIREE/Private/NNERuntimeIREERdg.cpp
--- a/Engine/Plugins/Experimental/NNERuntimeIREE/Source/NNERuntimeIREE/Private/NNERuntimeIREERdg.cpp
+++ b/Engine/Plugins/Experimental/NNERuntimeIREE/Source/NNERuntimeIREE/Private/NNERuntimeIREERdg.cpp
@@ -17,7 +17,7 @@ FString UNNERuntimeIREERdg::GetRuntimeName() const
 	return TEXT("NNERuntimeIREERdg");
 }
 
-bool UNNERuntimeIREERdg::CanCreateModelData(FString FileType, TConstArrayView<uint8> FileData, FGuid FileId, const ITargetPlatform* TargetPlatform) const
+bool UNNERuntimeIREERdg::CanCreateModelData(const FString& FileType, TConstArrayView<uint8> FileData, const TMap<FString, TConstArrayView<uint8>>& AdditionalFileData, const FGuid& FileId, const ITargetPlatform* TargetPlatform) const
 {
 #if WITH_EDITOR
 	return	FileType.Compare(TEXT("mlir"), ESearchCase::IgnoreCase) == 0;
@@ -26,12 +26,12 @@ bool UNNERuntimeIREERdg::CanCreateModelData(FString FileType, TConstArrayView<ui
 #endif // WITH_EDITOR
 }
 
-TSharedPtr<UE::NNE::FSharedModelData> UNNERuntimeIREERdg::CreateModelData(FString FileType, TConstArrayView<uint8> FileData, FGuid FileId, const ITargetPlatform* TargetPlatform)
+TSharedPtr<UE::NNE::FSharedModelData> UNNERuntimeIREERdg::CreateModelData(const FString& FileType, TConstArrayView<uint8> FileData, const TMap<FString, TConstArrayView<uint8>>& AdditionalFileData, const FGuid& FileId, const ITargetPlatform* TargetPlatform)
 {
 	return TSharedPtr<UE::NNE::FSharedModelData>();
 }
 
-FString UNNERuntimeIREERdg::GetModelDataIdentifier(FString FileType, TConstArrayView<uint8> FileData, FGuid FileId, const ITargetPlatform* TargetPlatform)
+FString UNNERuntimeIREERdg::GetModelDataIdentifier(const FString& FileType, TConstArrayView<uint8> FileData, const TMap<FString, TConstArrayView<uint8>>& AdditionalFileData, const FGuid& FileId, const ITargetPlatform* TargetPlatform) const
 {
 	FString PlatformName = TargetPlatform ? TargetPlatform->IniPlatformName() : UGameplayStatics::GetPlatformName();
 	return GetRuntimeName() + "-" + GUID.ToString(EGuidFormats::Digits) + "-" + FString::FromInt(UNNERuntimeIREERdg::Version) + "-" + FileId.ToString(EGuidFormats::Digits) + "-" + PlatformName;
@@ -96,9 +96,9 @@ bool UNNERuntimeIREERdg::IsAvailable() const
 
 FString UNNERuntimeIREERdg::GetRuntimeName() const { return TEXT(""); }
 
-bool UNNERuntimeIREERdg::CanCreateModelData(FString FileType, TConstArrayView<uint8> FileData, FGuid FileId, const ITargetPlatform* TargetPlatform) const { return false; };
-TSharedPtr<UE::NNE::FSharedModelData> UNNERuntimeIREERdg::CreateModelData(FString FileType, TConstArrayView<uint8> FileData, FGuid FileId, const ITargetPlatform* TargetPlatform) { return TSharedPtr<UE::NNE::FSharedModelData>(); };
-FString UNNERuntimeIREERdg::GetModelDataIdentifier(FString FileType, TConstArrayView<uint8> FileData, FGuid FileId, const ITargetPlatform* TargetPlatform) { return ""; };
+bool UNNERuntimeIREERdg::CanCreateModelData(const FString& FileType, TConstArrayView<uint8> FileData, const TMap<FString, TConstArrayView<uint8>>& AdditionalFileData, const FGuid& FileId, const ITargetPlatform* TargetPlatform) const { return false; };
+TSharedPtr<UE::NNE::FSharedModelData> UNNERuntimeIREERdg::CreateModelData(const FString& FileType, TConstArrayView<uint8> FileData, const TMap<FString, TConstArrayView<uint8>>& AdditionalFileData, const FGuid& FileId, const ITargetPlatform* TargetPlatform) { return TSharedPtr<UE::NNE::FSharedModelData>(); };
+FString UNNERuntimeIREERdg::GetModelDataIdentifier(const FString& FileType, TConstArrayView<uint8> FileData, const TMap<FString, TConstArrayView<uint8>>& AdditionalFileData, const FGuid& FileId, const ITargetPlatform* TargetPlatform) const { return ""; };
 
 bool UNNERuntimeIREERdg::CanCreateModelRDG(TObjectPtr<UNNEModelData> ModelData) const { return false; };
 TSharedPtr<UE::NNE::IModelRDG> UNNERuntimeIREERdg::CreateModelRDG(TObjectPtr<UNNEModelData> ModelData) { return TSharedPtr<UE::NNE::IModelRDG>(); };
diff --git a/Engine/Plugins/Experimental/NNERuntimeIREE/Source/NNERuntimeIREE/Private/NNERuntimeIREERdg.h b/Engine/Plugins/Experimental/NNERuntimeIREE/Source/NNERuntimeIREE/Private/NNERuntimeIREERdg.h
--- a/Engine/Plugins/Experimental/NNERuntimeIREE/Source/NNERuntimeIREE/Private/NNERuntimeIREERdg.h
+++ b/Engine/Plugins/Experimental/NNERuntimeIREE/Source/NNERuntimeIREE/Private/NNERuntimeIREERdg.h
@@ -24,9 +24,9 @@ public:
 
 	//~ Begin INNERuntime Interface
 	virtual FString GetRuntimeName() const override;
-	virtual bool CanCreateModelData(FString FileType, TConstArrayView<uint8> FileData, FGuid FileId, const ITargetPlatform* TargetPlatform) const override;
-	virtual TSharedPtr<UE::NNE::FSharedModelData> CreateModelData(FString FileType, TConstArrayView<uint8> FileData, FGuid FileId, const ITargetPlatform* TargetPlatform) override;
-	virtual FString GetModelDataIdentifier(FString FileType, TConstArrayView<uint8> FileData, FGuid FileId, const ITargetPlatform* TargetPlatform) override;
+	virtual bool CanCreateModelData(const FString& FileType, TConstArrayView<uint8> FileData, const TMap<FString, TConstArrayView<uint8>>& AdditionalFileData, const FGuid& FileId, const ITargetPlatform* TargetPlatform) const override;
+	virtual TSharedPtr<UE::NNE::FSharedModelData> CreateModelData(const FString& FileType, TConstArrayView<uint8> FileData, const TMap<FString, TConstArrayView<uint8>>& AdditionalFileData, const FGuid& FileId, const ITargetPlatform* TargetPlatform) override;
+	virtual FString GetModelDataIdentifier(const FString& FileType, TConstArrayView<uint8> FileData, const TMap<FString, TConstArrayView<uint8>>& AdditionalFileData, const FGuid& FileId, const ITargetPlatform* TargetPlatform) const override;
 	//~ End INNERuntime Interface
 
 	//~ Begin INNERuntimeRdg Interface
diff --git a/Engine/Plugins/Experimental/NNERuntimeORTCpu/Source/NNERuntimeORTCpu/Private/NNERuntimeORTCpu.cpp b/Engine/Plugins/Experimental/NNERuntimeORTCpu/Source/NNERuntimeORTCpu/Private/NNERuntimeORTCpu.cpp
--- a/Engine/Plugins/Experimental/NNERuntimeORTCpu/Source/NNERuntimeORTCpu/Private/NNERuntimeORTCpu.cpp
+++ b/Engine/Plugins/Experimental/NNERuntimeORTCpu/Source/NNERuntimeORTCpu/Private/NNERuntimeORTCpu.cpp
@@ -16,14 +16,14 @@
 FGuid UNNERuntimeORTCustomCpuImpl::GUID = FGuid((int32)'O', (int32)'C', (int32)'P', (int32)'U');
 int32 UNNERuntimeORTCustomCpuImpl::Version = 0x00000001;
 
-bool UNNERuntimeORTCustomCpuImpl::CanCreateModelData(FString FileType, TConstArrayView<uint8> FileData, FGuid FileId, const ITargetPlatform* TargetPlatform) const
+bool UNNERuntimeORTCustomCpuImpl::CanCreateModelData(const FString& FileType, TConstArrayView<uint8> FileData, const TMap<FString, TConstArrayView<uint8>>& AdditionalFileData, const FGuid& FileId, const ITargetPlatform* TargetPlatform) const
 {
 	return FileType.Compare("onnx", ESearchCase::IgnoreCase) == 0;
 }
 
-TSharedPtr<UE::NNE::FSharedModelData> UNNERuntimeORTCustomCpuImpl::CreateModelData(FString FileType, TConstArrayView<uint8> FileData, FGuid FileId, const ITargetPlatform* TargetPlatform)
+TSharedPtr<UE::NNE::FSharedModelData> UNNERuntimeORTCustomCpuImpl::CreateModelData(const FString& FileType, TConstArrayView<uint8> FileData, const TMap<FString, TConstArrayView<uint8>>& AdditionalFileData, const FGuid& FileId, const ITargetPlatform* TargetPlatform)
 {
-	if (!CanCreateModelData(FileType, FileData, FileId, TargetPlatform))
+	if (!CanCreateModelData(FileType, FileData, AdditionalFileData, FileId, TargetPlatform))
 	{
 		return {};
 	}
@@ -51,7 +51,7 @@ TSharedPtr<UE::NNE::FSharedModelData> UNNERuntimeORTCustomCpuImpl::CreateModelDa
 	return MakeShared<UE::NNE::FSharedModelData>(MakeSharedBufferFromArray(MoveTemp(Result)), 0);
 }
 
-FString UNNERuntimeORTCustomCpuImpl::GetModelDataIdentifier(FString FileType, TConstArrayView<uint8> FileData, FGuid FileId, const ITargetPlatform* TargetPlatform)
+FString UNNERuntimeORTCustomCpuImpl::GetModelDataIdentifier(const FString& FileType, TConstArrayView<uint8> FileData, const TMap<FString, TConstArrayView<uint8>>& AdditionalFileData, const FGuid& FileId, const ITargetPlatform* TargetPlatform) const
 {
 	return FileId.ToString(EGuidFormats::Digits) + "-" + UNNERuntimeORTCustomCpuImpl::GUID.ToString(EGuidFormats::Digits) + "-" + FString::FromInt(UNNERuntimeORTCustomCpuImpl::Version);
 }
diff --git a/Engine/Plugins/Experimental/NNERuntimeORTCpu/Source/NNERuntimeORTCpu/Private/NNERuntimeORTCpu.h b/Engine/Plugins/Experimental/NNERuntimeORTCpu/Source/NNERuntimeORTCpu/Private/NNERuntimeORTCpu.h
--- a/Engine/Plugins/Experimental/NNERuntimeORTCpu/Source/NNERuntimeORTCpu/Private/NNERuntimeORTCpu.h
+++ b/Engine/Plugins/Experimental/NNERuntimeORTCpu/Source/NNERuntimeORTCpu/Private/NNERuntimeORTCpu.h
@@ -33,9 +33,9 @@ public:
 		
 	virtual FString GetRuntimeName() const override { return TEXT("NNERuntimeORTCpu"); };
 
-	virtual bool CanCreateModelData(FString FileType, TConstArrayView<uint8> FileData, FGuid FileId, const ITargetPlatform* TargetPlatform) const override;
-	virtual TSharedPtr<UE::NNE::FSharedModelData> CreateModelData(FString FileType, TConstArrayView<uint8> FileData, FGuid FileId, const ITargetPlatform* TargetPlatform) override;
-	virtual FString GetModelDataIdentifier(FString FileType, TConstArrayView<uint8> FileData, FGuid FileId, const ITargetPlatform* TargetPlatform) override;
+	virtual bool CanCreateModelData(const FString& FileType, TConstArrayView<uint8> FileData, const TMap<FString, TConstArrayView<uint8>>& AdditionalFileData, const FGuid& FileId, const ITargetPlatform* TargetPlatform) const override;
+	virtual TSharedPtr<UE::NNE::FSharedModelData> CreateModelData(const FString& FileType, TConstArrayView<uint8> FileData, const TMap<FString, TConstArrayView<uint8>>& AdditionalFileData, const FGuid& FileId, const ITargetPlatform* TargetPlatform) override;
+	virtual FString GetModelDataIdentifier(const FString& FileType, TConstArrayView<uint8> FileData, const TMap<FString, TConstArrayView<uint8>>& AdditionalFileData, const FGuid& FileId, const ITargetPlatform* TargetPlatform) const override;
 
 	virtual bool CanCreateModelCPU(TObjectPtr<UNNEModelData> ModelData) const override;
 	virtual TSharedPtr<UE::NNE::IModelCPU> CreateModelCPU(TObjectPtr<UNNEModelData> ModelData) override;
diff --git a/Engine/Plugins/Experimental/NNERuntimeORTGpu/Source/NNERuntimeORTGpu/Private/NNERuntimeORTGpu.cpp b/Engine/Plugins/Experimental/NNERuntimeORTGpu/Source/NNERuntimeORTGpu/Private/NNERuntimeORTGpu.cpp
--- a/Engine/Plugins/Experimental/NNERuntimeORTGpu/Source/NNERuntimeORTGpu/Private/NNERuntimeORTGpu.cpp
+++ b/Engine/Plugins/Experimental/NNERuntimeORTGpu/Source/NNERuntimeORTGpu/Private/NNERuntimeORTGpu.cpp
@@ -18,14 +18,14 @@
 FGuid UNNERuntimeORTGpuImpl::GUID = FGuid((int32)'O', (int32)'G', (int32)'P', (int32)'U');
 int32 UNNERuntimeORTGpuImpl::Version = 0x00000001;
 
-bool UNNERuntimeORTGpuImpl::CanCreateModelData(FString FileType, TConstArrayView<uint8> FileData, FGuid FileId, const ITargetPlatform* TargetPlatform) const
+bool UNNERuntimeORTGpuImpl::CanCreateModelData(const FString& FileType, TConstArrayView<uint8> FileData, const TMap<FString, TConstArrayView<uint8>>& AdditionalFileData, const FGuid& FileId, const ITargetPlatform* TargetPlatform) const
 {
 	return FileType.Compare("onnx", ESearchCase::IgnoreCase) == 0;
 }
 
-TSharedPtr<UE::NNE::FSharedModelData> UNNERuntimeORTGpuImpl::CreateModelData(FString FileType, TConstArrayView<uint8> FileData, FGuid FileId, const ITargetPlatform* TargetPlatform)
+TSharedPtr<UE::NNE::FSharedModelData> UNNERuntimeORTGpuImpl::CreateModelData(const FString& FileType, TConstArrayView<uint8> FileData, const TMap<FString, TConstArrayView<uint8>>& AdditionalFileData, const FGuid& FileId, const ITargetPlatform* TargetPlatform)
 {
-	if (!CanCreateModelData(FileType, FileData, FileId, TargetPlatform))
+	if (!CanCreateModelData(FileType, FileData, AdditionalFileData, FileId, TargetPlatform))
 	{
 		return {};
 	}
@@ -53,7 +53,7 @@ TSharedPtr<UE::NNE::FSharedModelData> UNNERuntimeORTGpuImpl::CreateModelData(FSt
 	return MakeShared<UE::NNE::FSharedModelData>(MakeSharedBufferFromArray(MoveTemp(Result)), 0);
 }
 
-FString UNNERuntimeORTGpuImpl::GetModelDataIdentifier(FString FileType, TConstArrayView<uint8> FileData, FGuid FileId, const ITargetPlatform* TargetPlatform)
+FString UNNERuntimeORTGpuImpl::GetModelDataIdentifier(const FString& FileType, TConstArrayView<uint8> FileData, const TMap<FString, TConstArrayView<uint8>>& AdditionalFileData, const FGuid& FileId, const ITargetPlatform* TargetPlatform) const
 {
 	return FileId.ToString(EGuidFormats::Digits) + "-" + UNNERuntimeORTGpuImpl::GUID.ToString(EGuidFormats::Digits) + "-" + FString::FromInt(UNNERuntimeORTGpuImpl::Version);
 }
diff --git a/Engine/Plugins/Experimental/NNERuntimeORTGpu/Source/NNERuntimeORTGpu/Private/NNERuntimeORTGpu.h b/Engine/Plugins/Experimental/NNERuntimeORTGpu/Source/NNERuntimeORTGpu/Private/NNERuntimeORTGpu.h
--- a/Engine/Plugins/Experimental/NNERuntimeORTGpu/Source/NNERuntimeORTGpu/Private/NNERuntimeORTGpu.h
+++ b/Engine/Plugins/Experimental/NNERuntimeORTGpu/Source/NNERuntimeORTGpu/Private/NNERuntimeORTGpu.h
@@ -45,9 +45,9 @@ public:
 
 	virtual FString GetRuntimeName() const override;
 
-	virtual bool CanCreateModelData(FString FileType, TConstArrayView<uint8> FileData, FGuid FileId, const ITargetPlatform* TargetPlatform) const override;
-	virtual TSharedPtr<UE::NNE::FSharedModelData> CreateModelData(FString FileType, TConstArrayView<uint8> FileData, FGuid FileId, const ITargetPlatform* TargetPlatform) override;
-	virtual FString GetModelDataIdentifier(FString FileType, TConstArrayView<uint8> FileData, FGuid FileId, const ITargetPlatform* TargetPlatform) override;
+	virtual bool CanCreateModelData(const FString& FileType, TConstArrayView<uint8> FileData, const TMap<FString, TConstArrayView<uint8>>& AdditionalFileData, const FGuid& FileId, const ITargetPlatform* TargetPlatform) const override;
+	virtual TSharedPtr<UE::NNE::FSharedModelData> CreateModelData(const FString& FileType, TConstArrayView<uint8> FileData, const TMap<FString, TConstArrayView<uint8>>& AdditionalFileData, const FGuid& FileId, const ITargetPlatform* TargetPlatform) override;
+	virtual FString GetModelDataIdentifier(const FString& FileType, TConstArrayView<uint8> FileData, const TMap<FString, TConstArrayView<uint8>>& AdditionalFileData, const FGuid& FileId, const ITargetPlatform* TargetPlatform) const override;
 
 	virtual bool CanCreateModelGPU(TObjectPtr<UNNEModelData> ModelData) const override;
 	virtual TSharedPtr<UE::NNE::IModelGPU> CreateModelGPU(TObjectPtr<UNNEModelData> ModelData) override;
diff --git a/Engine/Plugins/Experimental/NNERuntimeRDG/Source/NNERuntimeRDG/Private/NNERuntimeRDGDml.cpp b/Engine/Plugins/Experimental/NNERuntimeRDG/Source/NNERuntimeRDG/Private/NNERuntimeRDGDml.cpp
--- a/Engine/Plugins/Experimental/NNERuntimeRDG/Source/NNERuntimeRDG/Private/NNERuntimeRDGDml.cpp
+++ b/Engine/Plugins/Experimental/NNERuntimeRDG/Source/NNERuntimeRDG/Private/NNERuntimeRDGDml.cpp
@@ -98,7 +98,7 @@ bool FRuntimeDmlStartup()
 
 using namespace UE::NNERuntimeRDG::Private::Dml;
 
-bool UNNERuntimeRDGDmlImpl::CanCreateModelData(FString FileType, TConstArrayView<uint8> FileData, FGuid FileId, const ITargetPlatform* TargetPlatform) const
+bool UNNERuntimeRDGDmlImpl::CanCreateModelData(const FString& FileType, TConstArrayView<uint8> FileData, const TMap<FString, TConstArrayView<uint8>>& AdditionalFileData, const FGuid& FileId, const ITargetPlatform* TargetPlatform) const
 {
 #ifdef NNE_UTILITIES_AVAILABLE
 	return FileType.Compare("onnx", ESearchCase::IgnoreCase) == 0;
@@ -108,9 +108,9 @@ bool UNNERuntimeRDGDmlImpl::CanCreateModelData(FString FileType, TConstArrayView
 #endif
 }
 
-TSharedPtr<UE::NNE::FSharedModelData> UNNERuntimeRDGDmlImpl::CreateModelData(FString FileType, TConstArrayView<uint8> FileData, FGuid FileId, const ITargetPlatform* TargetPlatform)
+TSharedPtr<UE::NNE::FSharedModelData> UNNERuntimeRDGDmlImpl::CreateModelData(const FString& FileType, TConstArrayView<uint8> FileData, const TMap<FString, TConstArrayView<uint8>>& AdditionalFileData, const FGuid& FileId, const ITargetPlatform* TargetPlatform)
 {
-	if (!CanCreateModelData(FileType, FileData, FileId, TargetPlatform))
+	if (!CanCreateModelData(FileType, FileData, AdditionalFileData, FileId, TargetPlatform))
 	{
 		return {};
 	}
@@ -147,7 +147,7 @@ TSharedPtr<UE::NNE::FSharedModelData> UNNERuntimeRDGDmlImpl::CreateModelData(FSt
 #endif //NNE_UTILITIES_AVAILABLE
 };
 
-FString UNNERuntimeRDGDmlImpl::GetModelDataIdentifier(FString FileType, TConstArrayView<uint8> FileData, FGuid FileId, const ITargetPlatform* TargetPlatform)
+FString UNNERuntimeRDGDmlImpl::GetModelDataIdentifier(const FString& FileType, TConstArrayView<uint8> FileData, const TMap<FString, TConstArrayView<uint8>>& AdditionalFileData, const FGuid& FileId, const ITargetPlatform* TargetPlatform) const
 {
 	return FileId.ToString(EGuidFormats::Digits) + "-" + FModelInfo::Get()->GetGuid().ToString(EGuidFormats::Digits) + "-" + FString::FromInt(FModelInfo::Get()->GetVersion());
 }
diff --git a/Engine/Plugins/Experimental/NNERuntimeRDG/Source/NNERuntimeRDG/Private/NNERuntimeRDGDml.h b/Engine/Plugins/Experimental/NNERuntimeRDG/Source/NNERuntimeRDG/Private/NNERuntimeRDGDml.h
--- a/Engine/Plugins/Experimental/NNERuntimeRDG/Source/NNERuntimeRDG/Private/NNERuntimeRDGDml.h
+++ b/Engine/Plugins/Experimental/NNERuntimeRDG/Source/NNERuntimeRDG/Private/NNERuntimeRDGDml.h
@@ -33,9 +33,9 @@ public:
 
 	virtual FString GetRuntimeName() const override { return TEXT("NNERuntimeRDGDml"); };
 
-	virtual bool CanCreateModelData(FString FileType, TConstArrayView<uint8> FileData, FGuid FileId, const ITargetPlatform* TargetPlatform) const override;
-	virtual TSharedPtr<UE::NNE::FSharedModelData> CreateModelData(FString FileType, TConstArrayView<uint8> FileData, FGuid FileId, const ITargetPlatform* TargetPlatform) override;
-	virtual FString GetModelDataIdentifier(FString FileType, TConstArrayView<uint8> FileData, FGuid FileId, const ITargetPlatform* TargetPlatform) override;
+	virtual bool CanCreateModelData(const FString& FileType, TConstArrayView<uint8> FileData, const TMap<FString, TConstArrayView<uint8>>& AdditionalFileData, const FGuid& FileId, const ITargetPlatform* TargetPlatform) const override;
+	virtual TSharedPtr<UE::NNE::FSharedModelData> CreateModelData(const FString& FileType, TConstArrayView<uint8> FileData, const TMap<FString, TConstArrayView<uint8>>& AdditionalFileData, const FGuid& FileId, const ITargetPlatform* TargetPlatform) override;
+	virtual FString GetModelDataIdentifier(const FString& FileType, TConstArrayView<uint8> FileData, const TMap<FString, TConstArrayView<uint8>>& AdditionalFileData, const FGuid& FileId, const ITargetPlatform* TargetPlatform) const override;
 
 	virtual bool CanCreateModelRDG(TObjectPtr<UNNEModelData> ModelData) const override;
 	virtual TSharedPtr<UE::NNE::IModelRDG> CreateModelRDG(TObjectPtr<UNNEModelData> ModelData) override;
diff --git a/Engine/Plugins/Experimental/NNERuntimeRDG/Source/NNERuntimeRDG/Private/NNERuntimeRDGHlsl.cpp b/Engine/Plugins/Experimental/NNERuntimeRDG/Source/NNERuntimeRDG/Private/NNERuntimeRDGHlsl.cpp
--- a/Engine/Plugins/Experimental/NNERuntimeRDG/Source/NNERuntimeRDG/Private/NNERuntimeRDGHlsl.cpp
+++ b/Engine/Plugins/Experimental/NNERuntimeRDG/Source/NNERuntimeRDG/Private/NNERuntimeRDGHlsl.cpp
@@ -85,7 +85,7 @@ bool UNNERuntimeRDGHlslImpl::Init()
 	return true;
 }
 
-bool UNNERuntimeRDGHlslImpl::CanCreateModelData(FString FileType, TConstArrayView<uint8> FileData, FGuid FileId, const ITargetPlatform* TargetPlatform) const
+bool UNNERuntimeRDGHlslImpl::CanCreateModelData(const FString& FileType, TConstArrayView<uint8> FileData, const TMap<FString, TConstArrayView<uint8>>& AdditionalFileData, const FGuid& FileId, const ITargetPlatform* TargetPlatform) const
 {
 #ifdef NNE_UTILITIES_AVAILABLE
 	return FileType.Compare("onnx", ESearchCase::IgnoreCase) == 0;
@@ -117,9 +117,9 @@ bool UNNERuntimeRDGHlslImpl::CanCreateModelRDG(TObjectPtr<UNNEModelData> ModelDa
 	return bResult;
 };
 
-TSharedPtr<UE::NNE::FSharedModelData> UNNERuntimeRDGHlslImpl::CreateModelData(FString FileType, TConstArrayView<uint8> FileData, FGuid FileId, const ITargetPlatform* TargetPlatform)
+TSharedPtr<UE::NNE::FSharedModelData> UNNERuntimeRDGHlslImpl::CreateModelData(const FString& FileType, TConstArrayView<uint8> FileData, const TMap<FString, TConstArrayView<uint8>>& AdditionalFileData, const FGuid& FileId, const ITargetPlatform* TargetPlatform)
 {
-	if (!CanCreateModelData(FileType, FileData, FileId, TargetPlatform))
+	if (!CanCreateModelData(FileType, FileData, AdditionalFileData, FileId, TargetPlatform))
 	{
 		return {};
 	}
@@ -151,7 +151,7 @@ TSharedPtr<UE::NNE::FSharedModelData> UNNERuntimeRDGHlslImpl::CreateModelData(FS
 #endif //NNE_UTILITIES_AVAILABLE
 };
 
-FString UNNERuntimeRDGHlslImpl::GetModelDataIdentifier(FString FileType, TConstArrayView<uint8> FileData, FGuid FileId, const ITargetPlatform* TargetPlatform)
+FString UNNERuntimeRDGHlslImpl::GetModelDataIdentifier(const FString& FileType, TConstArrayView<uint8> FileData, const TMap<FString, TConstArrayView<uint8>>& AdditionalFileData, const FGuid& FileId, const ITargetPlatform* TargetPlatform) const
 {
 	return FileId.ToString(EGuidFormats::Digits) + "-" + UNNERuntimeRDGHlslImpl::GUID.ToString(EGuidFormats::Digits) + "-" + FString::FromInt(UNNERuntimeRDGHlslImpl::Version);
 }
diff --git a/Engine/Plugins/Experimental/NNERuntimeRDG/Source/NNERuntimeRDG/Private/NNERuntimeRDGHlsl.h b/Engine/Plugins/Experimental/NNERuntimeRDG/Source/NNERuntimeRDG/Private/NNERuntimeRDGHlsl.h
--- a/Engine/Plugins/Experimental/NNERuntimeRDG/Source/NNERuntimeRDG/Private/NNERuntimeRDGHlsl.h
+++ b/Engine/Plugins/Experimental/NNERuntimeRDG/Source/NNERuntimeRDG/Private/NNERuntimeRDGHlsl.h
@@ -28,9 +28,9 @@ public:
 
 	virtual FString GetRuntimeName() const override { return TEXT("NNERuntimeRDGHlsl"); };
 
-	virtual bool CanCreateModelData(FString FileType, TConstArrayView<uint8> FileData, FGuid FileId, const ITargetPlatform* TargetPlatform) const override;
-	virtual TSharedPtr<UE::NNE::FSharedModelData> CreateModelData(FString FileType, TConstArrayView<uint8> FileData, FGuid FileId, const ITargetPlatform* TargetPlatform) override;
-	virtual FString GetModelDataIdentifier(FString FileType, TConstArrayView<uint8> FileData, FGuid FileId, const ITargetPlatform* TargetPlatform) override;
+	virtual bool CanCreateModelData(const FString& FileType, TConstArrayView<uint8> FileData, const TMap<FString, TConstArrayView<uint8>>& AdditionalFileData, const FGuid& FileId, const ITargetPlatform* TargetPlatform) const override;
+	virtual TSharedPtr<UE::NNE::FSharedModelData> CreateModelData(const FString& FileType, TConstArrayView<uint8> FileData, const TMap<FString, TConstArrayView<uint8>>& AdditionalFileData, const FGuid& FileId, const ITargetPlatform* TargetPlatform) override;
+	virtual FString GetModelDataIdentifier(const FString& FileType, TConstArrayView<uint8> FileData, const TMap<FString, TConstArrayView<uint8>>& AdditionalFileData, const FGuid& FileId, const ITargetPlatform* TargetPlatform) const override;
 
 	virtual bool CanCreateModelRDG(TObjectPtr<UNNEModelData> ModelData) const override;
 	virtual TSharedPtr<UE::NNE::IModelRDG> CreateModelRDG(TObjectPtr<UNNEModelData> ModelData) override;
diff --git a/Engine/Plugins/NNE/NNERuntimeORT/Source/NNERuntimeORT/Private/NNERuntimeORT.cpp b/Engine/Plugins/NNE/NNERuntimeORT/Source/NNERuntimeORT/Private/NNERuntimeORT.cpp
--- a/Engine/Plugins/NNE/NNERuntimeORT/Source/NNERuntimeORT/Private/NNERuntimeORT.cpp
+++ b/Engine/Plugins/NNE/NNERuntimeORT/Source/NNERuntimeORT/Private/NNERuntimeORT.cpp
@@ -24,19 +24,19 @@ int32 UNNERuntimeORTDml::Version = 0x00000001;
 FGuid UNNERuntimeORTCpu::GUID = FGuid((int32)'O', (int32)'C', (int32)'P', (int32)'U');
 int32 UNNERuntimeORTCpu::Version = 0x00000001;
 
-bool UNNERuntimeORTDml::CanCreateModelData(FString FileType, TConstArrayView<uint8> FileData, FGuid FileId, const ITargetPlatform* TargetPlatform) const
+bool UNNERuntimeORTDml::CanCreateModelData(const FString& FileType, TConstArrayView<uint8> FileData, const TMap<FString, TConstArrayView<uint8>>& AdditionalFileData, const FGuid& FileId, const ITargetPlatform* TargetPlatform) const
 {
 	return !FileData.IsEmpty() && FileType.Compare("onnx", ESearchCase::IgnoreCase) == 0;
 }
 
-bool UNNERuntimeORTCpu::CanCreateModelData(FString FileType, TConstArrayView<uint8> FileData, FGuid FileId, const ITargetPlatform* TargetPlatform) const
+bool UNNERuntimeORTCpu::CanCreateModelData(const FString& FileType, TConstArrayView<uint8> FileData, const TMap<FString, TConstArrayView<uint8>>& AdditionalFileData, const FGuid& FileId, const ITargetPlatform* TargetPlatform) const
 {
 	return !FileData.IsEmpty() && FileType.Compare("onnx", ESearchCase::IgnoreCase) == 0;
 }
 
-TSharedPtr<UE::NNE::FSharedModelData> UNNERuntimeORTDml::CreateModelData(FString FileType, TConstArrayView<uint8> FileData, FGuid FileId, const ITargetPlatform* TargetPlatform)
+TSharedPtr<UE::NNE::FSharedModelData> UNNERuntimeORTDml::CreateModelData(const FString& FileType, TConstArrayView<uint8> FileData, const TMap<FString, TConstArrayView<uint8>>& AdditionalFileData, const FGuid& FileId, const ITargetPlatform* TargetPlatform)
 {
-	if (!CanCreateModelData(FileType, FileData, FileId, TargetPlatform))
+	if (!CanCreateModelData(FileType, FileData, AdditionalFileData, FileId, TargetPlatform))
 	{
 		return {};
 	}
@@ -62,9 +62,9 @@ TSharedPtr<UE::NNE::FSharedModelData> UNNERuntimeORTDml::CreateModelData(FString
 	return MakeShared<UE::NNE::FSharedModelData>(MakeSharedBufferFromArray(MoveTemp(Result)), 0);
 }
 
-TSharedPtr<UE::NNE::FSharedModelData> UNNERuntimeORTCpu::CreateModelData(FString FileType, TConstArrayView<uint8> FileData, FGuid FileId, const ITargetPlatform* TargetPlatform)
+TSharedPtr<UE::NNE::FSharedModelData> UNNERuntimeORTCpu::CreateModelData(const FString& FileType, TConstArrayView<uint8> FileData, const TMap<FString, TConstArrayView<uint8>>& AdditionalFileData, const FGuid& FileId, const ITargetPlatform* TargetPlatform)
 {
-	if (!CanCreateModelData(FileType, FileData, FileId, TargetPlatform))
+	if (!CanCreateModelData(FileType, FileData, AdditionalFileData, FileId, TargetPlatform))
 	{
 		return {};
 	}
@@ -90,12 +90,12 @@ TSharedPtr<UE::NNE::FSharedModelData> UNNERuntimeORTCpu::CreateModelData(FString
 	return MakeShared<UE::NNE::FSharedModelData>(MakeSharedBufferFromArray(MoveTemp(Result)), 0);
 }
 
-FString UNNERuntimeORTDml::GetModelDataIdentifier(FString FileType, TConstArrayView<uint8> FileData, FGuid FileId, const ITargetPlatform* TargetPlatform)
+FString UNNERuntimeORTDml::GetModelDataIdentifier(const FString& FileType, TConstArrayView<uint8> FileData, const TMap<FString, TConstArrayView<uint8>>& AdditionalFileData, const FGuid& FileId, const ITargetPlatform* TargetPlatform) const
 {
 	return FileId.ToString(EGuidFormats::Digits) + "-" + UNNERuntimeORTDml::GUID.ToString(EGuidFormats::Digits) + "-" + FString::FromInt(UNNERuntimeORTDml::Version);
 }
 
-FString UNNERuntimeORTCpu::GetModelDataIdentifier(FString FileType, TConstArrayView<uint8> FileData, FGuid FileId, const ITargetPlatform* TargetPlatform)
+FString UNNERuntimeORTCpu::GetModelDataIdentifier(const FString& FileType, TConstArrayView<uint8> FileData, const TMap<FString, TConstArrayView<uint8>>& AdditionalFileData, const FGuid& FileId, const ITargetPlatform* TargetPlatform) const
 {
 	return FileId.ToString(EGuidFormats::Digits) + "-" + UNNERuntimeORTCpu::GUID.ToString(EGuidFormats::Digits) + "-" + FString::FromInt(UNNERuntimeORTCpu::Version);
 }
diff --git a/Engine/Plugins/NNE/NNERuntimeORT/Source/NNERuntimeORT/Private/NNERuntimeORT.h b/Engine/Plugins/NNE/NNERuntimeORT/Source/NNERuntimeORT/Private/NNERuntimeORT.h
--- a/Engine/Plugins/NNE/NNERuntimeORT/Source/NNERuntimeORT/Private/NNERuntimeORT.h
+++ b/Engine/Plugins/NNE/NNERuntimeORT/Source/NNERuntimeORT/Private/NNERuntimeORT.h
@@ -34,9 +34,9 @@ public:
 
 	virtual FString GetRuntimeName() const override;
 
-	virtual bool CanCreateModelData(FString FileType, TConstArrayView<uint8> FileData, FGuid FileId, const ITargetPlatform* TargetPlatform) const override;
-	virtual TSharedPtr<UE::NNE::FSharedModelData> CreateModelData(FString FileType, TConstArrayView<uint8> FileData, FGuid FileId, const ITargetPlatform* TargetPlatform) override;
-	virtual FString GetModelDataIdentifier(FString FileType, TConstArrayView<uint8> FileData, FGuid FileId, const ITargetPlatform* TargetPlatform) override;
+	virtual bool CanCreateModelData(const FString& FileType, TConstArrayView<uint8> FileData, const TMap<FString, TConstArrayView<uint8>>& AdditionalFileData, const FGuid& FileId, const ITargetPlatform* TargetPlatform) const override;
+	virtual TSharedPtr<UE::NNE::FSharedModelData> CreateModelData(const FString& FileType, TConstArrayView<uint8> FileData, const TMap<FString, TConstArrayView<uint8>>& AdditionalFileData, const FGuid& FileId, const ITargetPlatform* TargetPlatform) override;
+	virtual FString GetModelDataIdentifier(const FString& FileType, TConstArrayView<uint8> FileData, const TMap<FString, TConstArrayView<uint8>>& AdditionalFileData, const FGuid& FileId, const ITargetPlatform* TargetPlatform) const override;
 
 	virtual bool CanCreateModelGPU(TObjectPtr<UNNEModelData> ModelData) const override;
 	virtual TSharedPtr<UE::NNE::IModelGPU> CreateModelGPU(TObjectPtr<UNNEModelData> ModelData) override;
@@ -62,9 +62,9 @@ public:
 		
 	virtual FString GetRuntimeName() const override;
 	
-	virtual bool CanCreateModelData(FString FileType, TConstArrayView<uint8> FileData, FGuid FileId, const ITargetPlatform* TargetPlatform) const override;
-	virtual TSharedPtr<UE::NNE::FSharedModelData> CreateModelData(FString FileType, TConstArrayView<uint8> FileData, FGuid FileId, const ITargetPlatform* TargetPlatform) override;
-	virtual FString GetModelDataIdentifier(FString FileType, TConstArrayView<uint8> FileData, FGuid FileId, const ITargetPlatform* TargetPlatform) override;
+	virtual bool CanCreateModelData(const FString& FileType, TConstArrayView<uint8> FileData, const TMap<FString, TConstArrayView<uint8>>& AdditionalFileData, const FGuid& FileId, const ITargetPlatform* TargetPlatform) const override;
+	virtual TSharedPtr<UE::NNE::FSharedModelData> CreateModelData(const FString& FileType, TConstArrayView<uint8> FileData, const TMap<FString, TConstArrayView<uint8>>& AdditionalFileData, const FGuid& FileId, const ITargetPlatform* TargetPlatform) override;
+	virtual FString GetModelDataIdentifier(const FString& FileType, TConstArrayView<uint8> FileData, const TMap<FString, TConstArrayView<uint8>>& AdditionalFileData, const FGuid& FileId, const ITargetPlatform* TargetPlatform) const override;
 
 	virtual bool CanCreateModelCPU(TObjectPtr<UNNEModelData> ModelData) const override;
 	virtual TSharedPtr<UE::NNE::IModelCPU> CreateModelCPU(TObjectPtr<UNNEModelData> ModelData) override;
diff --git a/Engine/Source/Runtime/NNE/Private/NNEModelData.cpp b/Engine/Source/Runtime/NNE/Private/NNEModelData.cpp
--- a/Engine/Source/Runtime/NNE/Private/NNEModelData.cpp
+++ b/Engine/Source/Runtime/NNE/Private/NNEModelData.cpp
@@ -27,6 +27,7 @@ namespace UE::NNE::ModelData
 		V0 = 0, // Initial
 		V1 = 1, // TargetRuntimes and AssetImportData
 		V2 = 2, // Re-arrange fields and store only ModelData in cooked assets
+		V3 = 3, // Adding AdditionalFileData
 		// New versions can be added above this line
 		VersionPlusOne,
 		Latest = VersionPlusOne - 1
@@ -130,14 +131,14 @@ namespace UE::NNE::ModelData
 
 #endif // WITH_EDITOR
 
-	inline TSharedPtr<UE::NNE::FSharedModelData> CreateModelData(const FString& RuntimeName, FString FileType, const TArray<uint8>& FileData, FGuid FileId, const ITargetPlatform* TargetPlatform)
+	inline TSharedPtr<UE::NNE::FSharedModelData> CreateModelData(const FString& RuntimeName, const FString& FileType, const TArray<uint8>& FileData, const TMap<FString, TConstArrayView<uint8>>& AdditionalFileData, const FGuid& FileId, const ITargetPlatform* TargetPlatform)
 	{
 		TWeakInterfacePtr<INNERuntime> NNERuntime = UE::NNE::GetRuntime<INNERuntime>(RuntimeName);
 		if (NNERuntime.IsValid())
 		{
-			if (NNERuntime->CanCreateModelData(FileType, FileData, FileId, TargetPlatform))
+			if (NNERuntime->CanCreateModelData(FileType, FileData, AdditionalFileData, FileId, TargetPlatform))
 			{
-				return NNERuntime->CreateModelData(FileType, FileData, FileId, TargetPlatform);
+				return NNERuntime->CreateModelData(FileType, FileData, AdditionalFileData, FileId, TargetPlatform);
 			}
 		}
 		else
@@ -218,13 +219,15 @@ void UNNEModelData::Serialize(FArchive& Ar)
 		bool bWriteModelData = true;
 		if (Ar.IsCooking())
 		{
-			// Optimize storage: FileData is not required anymore because we have the model and can cook it for every runtime
+			// Optimize storage: FileData and AdditionalFileData are not required anymore because we have the model and can cook it for every runtime
 			TArray<FString> TmpTargetRuntimes;
 			Ar << TmpTargetRuntimes;
 			FString TmpFileType;
 			Ar << TmpFileType;
 			TArray<uint8> TmpFileData;
 			Ar << TmpFileData;
+			int32 NumAdditionalFileDataItems = 0;
+			Ar << NumAdditionalFileDataItems;
 			Ar << FileId;
 
 			// Cooking must recreate all model data but only if file data is still available
@@ -246,6 +249,14 @@ void UNNEModelData::Serialize(FArchive& Ar)
 					CookRuntimeNames.Append(GetTargetRuntimes());
 				}
 
+				TMap<FString, TConstArrayView<uint8>> AdditionalFileDataView;
+				TSet<FString> Keys;
+				AdditionalFileData.GetKeys(Keys);
+				for (auto& Key : Keys)
+				{
+					AdditionalFileDataView[Key] = AdditionalFileData[Key];
+				}
+
 				for (const FString& RuntimeName : CookRuntimeNames)
 				{
 					TSharedPtr<UE::NNE::FSharedModelData> SharedModelData;
@@ -254,7 +265,7 @@ void UNNEModelData::Serialize(FArchive& Ar)
 					FString ModelDataIdentifier;
 					if (NNERuntime.IsValid())
 					{
-						ModelDataIdentifier = NNERuntime->GetModelDataIdentifier(FileType, FileData, FileId, Ar.GetArchiveState().CookingTarget());
+						ModelDataIdentifier = NNERuntime->GetModelDataIdentifier(FileType, FileData, AdditionalFileDataView, FileId, Ar.GetArchiveState().CookingTarget());
 						if (ModelDataIdentifier.Len() > 0)
 						{
 							if (UE::NNE::ConsoleVariables::CVarNNEEditorUseDDC.GetValueOnGameThread() &&
@@ -275,7 +286,7 @@ void UNNEModelData::Serialize(FArchive& Ar)
 #endif //WITH_EDITOR
 					if (!SharedModelData.IsValid() || SharedModelData->GetView().Num() == 0)
 					{
-						SharedModelData = UE::NNE::ModelData::CreateModelData(RuntimeName, FileType, FileData, FileId, Ar.GetArchiveState().CookingTarget());
+						SharedModelData = UE::NNE::ModelData::CreateModelData(RuntimeName, FileType, FileData, AdditionalFileDataView, FileId, Ar.GetArchiveState().CookingTarget());
 #if WITH_EDITOR
 						if (SharedModelData.IsValid() && (SharedModelData->GetView().Num() > 0) && NNERuntime.IsValid() && (ModelDataIdentifier.Len() > 0))
 						{
@@ -296,6 +307,17 @@ void UNNEModelData::Serialize(FArchive& Ar)
 			Ar << TargetRuntimes;
 			Ar << FileType;
 			Ar << FileData;
+
+			int32 NumAdditionalFileDataItems = AdditionalFileData.Num();
+			Ar << NumAdditionalFileDataItems;
+			TSet<FString> Keys;
+			AdditionalFileData.GetKeys(Keys);
+			for (auto& Key : Keys)
+			{
+				Ar << Key;
+				Ar << AdditionalFileData[Key];
+			}
+
 			Ar << FileId;
 
 #if WITH_EDITOR
@@ -334,6 +356,7 @@ void UNNEModelData::Serialize(FArchive& Ar)
 	{
 		TObjectPtr<class UAssetImportData> AssetImportData;
 		int32 NumItems = 0;
+		int32 NumAdditionalFileDataItems = 0;
 		FString Name;
 		uint32 MemoryAlignment = 0;
 		uint64 DataSize = 0;
@@ -395,6 +418,31 @@ void UNNEModelData::Serialize(FArchive& Ar)
 			}
 			break;
 
+		case UE::NNE::ModelData::Version::V3:
+			Ar << TargetRuntimes;
+			Ar << FileType;
+			Ar << FileData;
+			Ar << NumAdditionalFileDataItems;
+			AdditionalFileData.Empty();
+			for (Index = 0; Index < NumItems; Index++)
+			{
+				Ar << Name;
+				Ar << Data;
+				AdditionalFileData.Add(Name, Data);
+			}
+			Ar << FileId;
+			Ar << NumItems;
+			for (Index = 0; Index < NumItems; Index++)
+			{
+				Ar << Name;
+				Ar << MemoryAlignment;
+				Ar << DataSize;
+				RawData = FMemory::Malloc(DataSize, MemoryAlignment);
+				Ar.Serialize(RawData, DataSize);
+				ModelData.Add(Name, MakeShared<UE::NNE::FSharedModelData>(FSharedBuffer::TakeOwnership(RawData, DataSize, FMemory::Free), MemoryAlignment));
+			}
+			break;
+
 		default:
 			UE_LOG(LogNNE, Error, TEXT("UNNEModelData: Unknown asset version %d: Deserialisation failed, please reimport the original model."), Ar.CustomVer(UE::NNE::ModelData::GUID));
 			break;
@@ -402,11 +450,16 @@ void UNNEModelData::Serialize(FArchive& Ar)
 	}
 }
 
-void UNNEModelData::Init(const FString& Type, TConstArrayView<uint8> Buffer)
+void UNNEModelData::Init(const FString& Type, TConstArrayView<uint8> Buffer, const TMap<FString, TConstArrayView<uint8>>& AdditionalBuffers)
 {
 	TargetRuntimes.Empty();
 	FileType = Type;
 	FileData = Buffer;
+	AdditionalFileData.Empty();
+	for (auto& Element : AdditionalBuffers)
+	{
+		AdditionalFileData[Element.Key] = AdditionalBuffers[Element.Key];
+	}
 	FPlatformMisc::CreateGuid(FileId);
 	ModelData.Empty();
 }
@@ -445,10 +498,16 @@ TConstArrayView<uint8> UNNEModelData::GetFileData() const
 	return FileData;
 }
 
+TConstArrayView<uint8> UNNEModelData::GetAdditionalFileData(const FString& Key) const
+{
+	return AdditionalFileData.Contains(Key) ? AdditionalFileData[Key] : TConstArrayView<uint8>();
+}
+
 void UNNEModelData::ClearFileDataAndFileType()
 {
 	FileType = "";
 	FileData.Empty();
+	AdditionalFileData.Empty();
 }
 
 FGuid UNNEModelData::GetFileId() const
@@ -484,6 +543,14 @@ TSharedPtr<UE::NNE::FSharedModelData> UNNEModelData::GetModelData(const FString&
 		return TSharedPtr<UE::NNE::FSharedModelData>();
 	}
 
+	TMap<FString, TConstArrayView<uint8>> AdditionalFileDataView;
+	TSet<FString> Keys;
+	AdditionalFileData.GetKeys(Keys);
+	for (auto& Key : Keys)
+	{
+		AdditionalFileDataView[Key] = AdditionalFileData[Key];
+	}
+
 #if WITH_EDITOR
 	TWeakInterfacePtr<INNERuntime> NNERuntime = UE::NNE::GetRuntime<INNERuntime>(RuntimeName);
 	if (!NNERuntime.IsValid())
@@ -492,7 +559,7 @@ TSharedPtr<UE::NNE::FSharedModelData> UNNEModelData::GetModelData(const FString&
 		return TSharedPtr<UE::NNE::FSharedModelData>();
 	}
 
-	FString ModelDataIdentifier = NNERuntime->GetModelDataIdentifier(FileType, FileData, FileId, nullptr);
+	FString ModelDataIdentifier = NNERuntime->GetModelDataIdentifier(FileType, FileData, AdditionalFileDataView, FileId, nullptr);
 	if (ModelDataIdentifier.Len() == 0)
 	{
 		UE_LOG(LogNNE, Error, TEXT("UNNEModelData: Runtime '%s' returned an empty string as a ModelDataIdentifier. GetModelDataIdentifier should always return a valid identifier."), *RuntimeName);
@@ -512,7 +579,7 @@ TSharedPtr<UE::NNE::FSharedModelData> UNNEModelData::GetModelData(const FString&
 #endif //WITH_EDITOR
 
 	// Try to create the model
-	TSharedPtr<UE::NNE::FSharedModelData> CreatedData = UE::NNE::ModelData::CreateModelData(RuntimeName, FileType, FileData, FileId, nullptr);
+	TSharedPtr<UE::NNE::FSharedModelData> CreatedData = UE::NNE::ModelData::CreateModelData(RuntimeName, FileType, FileData, AdditionalFileDataView, FileId, nullptr);
 	if (!CreatedData.IsValid() || CreatedData->GetView().Num() < 1)
 	{
 		return TSharedPtr<UE::NNE::FSharedModelData>();
diff --git a/Engine/Source/Runtime/NNE/Public/NNEModelData.h b/Engine/Source/Runtime/NNE/Public/NNEModelData.h
--- a/Engine/Source/Runtime/NNE/Public/NNEModelData.h
+++ b/Engine/Source/Runtime/NNE/Public/NNEModelData.h
@@ -89,8 +89,9 @@ public:
 	 *
 	 * @param Type A string identifying the type of data inside this asset. Corresponds to the extension of the imported file.
 	 * @param Buffer The raw binary file data of the imported model to be copied into this asset.
+	 * @param AdditionalBuffers Additional raw binary data of the model to be copied into this asset.
 	 */
-	void Init(const FString& Type, TConstArrayView<uint8> Buffer);
+	void Init(const FString& Type, TConstArrayView<uint8> Buffer, const TMap<FString, TConstArrayView<uint8>>& AdditionalBuffers = TMap<FString, TConstArrayView<uint8>>());
 
 	/**
 	 * In editor: Get the target runtimes this model data will be cooked for. An empty list means all runtimes.
@@ -127,6 +128,16 @@ public:
 	 */
 	TConstArrayView<uint8> GetFileData() const;
 
+	/**
+	 * Get read only access to AdditionalFileData.
+	 *
+	 * In editor: The AdditionalFileData contains the additional binary data of the neural network model.
+	 * In standalone: An empty map.
+	 *
+	 * @return The AdditionalFileData with a given Key if it exists and an empty view in standalone or when the key does not exist.
+	 */
+	TConstArrayView<uint8> GetAdditionalFileData(const FString& Key) const;
+
 	/**
 	 * Clears the FileData and the FileType.
 	 *
@@ -176,6 +187,11 @@ private:
 	 */
 	TArray<uint8> FileData;
 
+	/**
+	 * Additional raw binary data of the imported model.
+	 */
+	TMap<FString, TArray<uint8>> AdditionalFileData;
+
 	/**
 	 * A Guid that uniquely identifies this model. This is used to cache optimized models in the editor.
 	 */
diff --git a/Engine/Source/Runtime/NNE/Public/NNERuntime.h b/Engine/Source/Runtime/NNE/Public/NNERuntime.h
--- a/Engine/Source/Runtime/NNE/Public/NNERuntime.h
+++ b/Engine/Source/Runtime/NNE/Public/NNERuntime.h
@@ -45,31 +45,34 @@ public:
 	 *
 	 * @param FileType The type of file inside FileData. Corresponds to the file extension (e.g. 'onnx').
 	 * @param FileData The raw binary file of a neural network model.
+	 * @param AdditionalFileData Additional raw binary data of the neural network model.
 	 * @param FileId The unique identifier representing FileData.
 	 * @param TargetPlatform The Interface identifying the target platform for which the data needs to be created. A null pointer indicates the currently compiled/running platform.
 	 * @return True if the runtime is able to create model data, false otherwise.
 	 */
-	virtual bool CanCreateModelData(FString FileType, TConstArrayView<uint8> FileData, FGuid FileId, const ITargetPlatform* TargetPlatform) const = 0;
+	virtual bool CanCreateModelData(const FString& FileType, TConstArrayView<uint8> FileData, const TMap<FString, TConstArrayView<uint8>>& AdditionalFileData, const FGuid& FileId, const ITargetPlatform* TargetPlatform) const = 0;
 
 	/**
 	 * Create model data given some raw file data.
 	 *
 	 * @param FileType The type of file inside FileData. Corresponds to the file extension (e.g. 'onnx').
 	 * @param FileData The raw binary file of a neural network model.
+	 * @param AdditionalFileData Additional raw binary data of the neural network model.
 	 * @param FileId The unique identifier representing FileData.
 	 * @param TargetPlatform The Interface identifying the target platform for which the data needs to be created. A null pointer indicates the currently compiled/running platform.
 	 * @return A pointer to data representing the runtime specific representation of the model to be stored by UNNEModelData on success or an invalid pointer otherwise.
 	 */
-	virtual TSharedPtr<UE::NNE::FSharedModelData> CreateModelData(FString FileType, TConstArrayView<uint8> FileData, FGuid FileId, const ITargetPlatform* TargetPlatform) = 0;
+	virtual TSharedPtr<UE::NNE::FSharedModelData> CreateModelData(const FString& FileType, TConstArrayView<uint8> FileData, const TMap<FString, TConstArrayView<uint8>>& AdditionalFileData, const FGuid& FileId, const ITargetPlatform* TargetPlatform) = 0;
 
 	/**
 	 * Get an id uniquely identifying the model data. This is used by UNNEModelData to create a FCacheKey to get and add the model data from and to a DDC.
 	 *
 	 * @param FileType The type of file inside FileData. Corresponds to the file extension (e.g. 'onnx').
 	 * @param FileData The raw binary file of a neural network model.
+	 * @param AdditionalFileData Additional raw binary data of the neural network model.
 	 * @param FileId The unique identifier representing FileData.
 	 * @param TargetPlatform The Interface identifying the target platform for which the data is produced or by which it will be consumed. A null pointer indicates the currently compiled/running platform.
 	 * @return A string that identifies the model data uniquely, e.g. by containing the FileId, model data version and target platform display name.
 	 */
-	virtual FString GetModelDataIdentifier(FString FileType, TConstArrayView<uint8> FileData, FGuid FileId, const ITargetPlatform* TargetPlatform) = 0;
+	virtual FString GetModelDataIdentifier(const FString& FileType, TConstArrayView<uint8> FileData, const TMap<FString, TConstArrayView<uint8>>& AdditionalFileData, const FGuid& FileId, const ITargetPlatform* TargetPlatform) const = 0;
 };
\ No newline at end of file
