[NNE] Fixing a minor issue with a loop variable

#rnx

diff --git a/Engine/Source/Runtime/NNE/Private/NNEModelData.cpp b/Engine/Source/Runtime/NNE/Private/NNEModelData.cpp
--- a/Engine/Source/Runtime/NNE/Private/NNEModelData.cpp
+++ b/Engine/Source/Runtime/NNE/Private/NNEModelData.cpp
@@ -424,7 +424,7 @@ void UNNEModelData::Serialize(FArchive& Ar)
 			Ar << FileData;
 			Ar << NumAdditionalFileDataItems;
 			AdditionalFileData.Empty();
-			for (Index = 0; Index < NumItems; Index++)
+			for (Index = 0; Index < NumAdditionalFileDataItems; Index++)
 			{
 				Ar << Name;
 				Ar << Data;
