[NNE] Remove UNNEModelData  guid from cooked asset

Problems:
* If one create UNNEModelData on the fly cooked binary would contain the guid witch would thus make the cooked data undeterministic.
* AdditionalFileData would not work as the key are not already existing in the map.

Solutions:
* Remove guid from cook build
* Emplace the new element to the AdditionalFileData map

Test: Win64 Editor: smoke tests
Test: Win64 Standalone : all tests

#rnx

diff --git a/Engine/Source/Runtime/NNE/Private/NNEModelData.cpp b/Engine/Source/Runtime/NNE/Private/NNEModelData.cpp
--- a/Engine/Source/Runtime/NNE/Private/NNEModelData.cpp
+++ b/Engine/Source/Runtime/NNE/Private/NNEModelData.cpp
@@ -228,7 +228,8 @@ void UNNEModelData::Serialize(FArchive& Ar)
 			Ar << TmpFileData;
 			int32 NumAdditionalFileDataItems = 0;
 			Ar << NumAdditionalFileDataItems;
-			Ar << FileId;
+			FGuid TmpGuid;
+			Ar << TmpGuid;
 
 			// Cooking must recreate all model data but only if file data is still available
 			if (FileData.Num() > 0)
@@ -254,7 +255,7 @@ void UNNEModelData::Serialize(FArchive& Ar)
 				AdditionalFileData.GetKeys(Keys);
 				for (auto& Key : Keys)
 				{
-					AdditionalFileDataView[Key] = AdditionalFileData[Key];
+					AdditionalFileDataView.Emplace(Key, AdditionalFileData[Key]);
 				}
 
 				for (const FString& RuntimeName : CookRuntimeNames)
@@ -458,7 +459,7 @@ void UNNEModelData::Init(const FString& Type, TConstArrayView<uint8> Buffer, con
 	AdditionalFileData.Empty();
 	for (auto& Element : AdditionalBuffers)
 	{
-		AdditionalFileData[Element.Key] = AdditionalBuffers[Element.Key];
+		AdditionalFileData.Emplace(Element.Key, AdditionalBuffers[Element.Key]);
 	}
 	FPlatformMisc::CreateGuid(FileId);
 	ModelData.Empty();
@@ -548,7 +549,7 @@ TSharedPtr<UE::NNE::FSharedModelData> UNNEModelData::GetModelData(const FString&
 	AdditionalFileData.GetKeys(Keys);
 	for (auto& Key : Keys)
 	{
-		AdditionalFileDataView[Key] = AdditionalFileData[Key];
+		AdditionalFileDataView.Emplace(Key, AdditionalFileData[Key]);
 	}
 
 #if WITH_EDITOR
diff --git a/Engine/Source/Runtime/NNE/Public/NNEModelData.h b/Engine/Source/Runtime/NNE/Public/NNEModelData.h
--- a/Engine/Source/Runtime/NNE/Public/NNEModelData.h
+++ b/Engine/Source/Runtime/NNE/Public/NNEModelData.h
@@ -149,6 +149,7 @@ public:
 	 * Get the FGuid identifying the FileData.
 	 *
 	 * The FileId is created on import of an asset. It can be used to identify the FileData, e.g. when putting corresponding data into the DDC or caching data locally.
+	 * In standalone: An empty FGuid.
 	 *
 	 * @return The FileId.
 	 */
