[NNE] Core cleanup: forward declarations, includes and minor code formatting.

Tests:
- Compile Files (Non-Unity)
- Unreal Editor, Win64: Style transfer and smoke tests
- Unreal Game, Win64: Style transfer and smoke tests

#rnx

diff --git a/Engine/Source/Runtime/NNE/Internal/NNEAttributeDataType.h b/Engine/Source/Runtime/NNE/Internal/NNEAttributeDataType.h
--- a/Engine/Source/Runtime/NNE/Internal/NNEAttributeDataType.h
+++ b/Engine/Source/Runtime/NNE/Internal/NNEAttributeDataType.h
@@ -2,7 +2,6 @@
 
 #pragma once
 
-#include "UObject/Class.h"
 #include "UObject/Object.h"
 
 #include "NNEAttributeDataType.generated.h"
diff --git a/Engine/Source/Runtime/NNE/Internal/NNEAttributeValue.h b/Engine/Source/Runtime/NNE/Internal/NNEAttributeValue.h
--- a/Engine/Source/Runtime/NNE/Internal/NNEAttributeValue.h
+++ b/Engine/Source/Runtime/NNE/Internal/NNEAttributeValue.h
@@ -5,7 +5,6 @@
 #include "NNEAttributeValueTraits.h"
 #include "Serialization/MemoryWriter.h"
 #include "Serialization/MemoryReader.h"
-#include "UObject/Class.h"
 #include "UObject/Object.h"
 
 #include "NNEAttributeValue.generated.h"
diff --git a/Engine/Source/Runtime/NNE/Internal/NNEModelBase.h b/Engine/Source/Runtime/NNE/Internal/NNEModelBase.h
--- a/Engine/Source/Runtime/NNE/Internal/NNEModelBase.h
+++ b/Engine/Source/Runtime/NNE/Internal/NNEModelBase.h
@@ -2,7 +2,7 @@
 
 #pragma once
 
-#include "CoreMinimal.h"
+#include "Containers/ContainersFwd.h"
 #include "NNE.h"
 #include "NNETypes.h"
 
diff --git a/Engine/Source/Runtime/NNE/Internal/NNETensor.h b/Engine/Source/Runtime/NNE/Internal/NNETensor.h
--- a/Engine/Source/Runtime/NNE/Internal/NNETensor.h
+++ b/Engine/Source/Runtime/NNE/Internal/NNETensor.h
@@ -2,7 +2,8 @@
 
 #pragma once
 
-#include "CoreMinimal.h"
+#include "Containers/ContainersFwd.h"
+#include "Containers/UnrealString.h"
 #include "NNETypes.h"
 
 namespace UE::NNE::Internal
diff --git a/Engine/Source/Runtime/NNE/Public/NNE.h b/Engine/Source/Runtime/NNE/Public/NNE.h
--- a/Engine/Source/Runtime/NNE/Public/NNE.h
+++ b/Engine/Source/Runtime/NNE/Public/NNE.h
@@ -2,7 +2,8 @@
 
 #pragma once
 
-#include "CoreMinimal.h"
+#include "Containers/ArrayView.h"
+#include "Containers/UnrealString.h"
 #include "NNERuntime.h"
 #include "UObject/WeakInterfacePtr.h"
 
@@ -52,7 +53,6 @@ namespace UE::NNE
 				T* RuntimePtr = Cast<T>(Runtime.Get());
 				return TWeakInterfacePtr<T>(RuntimePtr);
 			}
-				
 		}
 		return TWeakInterfacePtr<T>(nullptr);
 	}
diff --git a/Engine/Source/Runtime/NNE/Public/NNEModelData.h b/Engine/Source/Runtime/NNE/Public/NNEModelData.h
--- a/Engine/Source/Runtime/NNE/Public/NNEModelData.h
+++ b/Engine/Source/Runtime/NNE/Public/NNEModelData.h
@@ -2,7 +2,8 @@
 
 #pragma once
 
-#include "CoreMinimal.h"
+#include "Containers/ContainersFwd.h"
+#include "Containers/UnrealString.h"
 #include "Memory/SharedBuffer.h"
 #include "Misc/Guid.h"
 #include "Templates/SharedPointer.h"
diff --git a/Engine/Source/Runtime/NNE/Public/NNERuntime.h b/Engine/Source/Runtime/NNE/Public/NNERuntime.h
--- a/Engine/Source/Runtime/NNE/Public/NNERuntime.h
+++ b/Engine/Source/Runtime/NNE/Public/NNERuntime.h
@@ -2,7 +2,8 @@
 
 #pragma once
 
-#include "CoreMinimal.h"
+#include "Containers/ContainersFwd.h"
+#include "Containers/UnrealString.h"
 #include "UObject/Interface.h"
 
 #include "NNERuntime.generated.h"
diff --git a/Engine/Source/Runtime/NNE/Public/NNERuntimeCPU.h b/Engine/Source/Runtime/NNE/Public/NNERuntimeCPU.h
--- a/Engine/Source/Runtime/NNE/Public/NNERuntimeCPU.h
+++ b/Engine/Source/Runtime/NNE/Public/NNERuntimeCPU.h
@@ -2,8 +2,7 @@
 
 #pragma once
 
-#include "CoreMinimal.h"
-#include "NNERuntime.h"
+#include "Containers/ContainersFwd.h"
 #include "NNETypes.h"
 #include "UObject/Interface.h"
 
diff --git a/Engine/Source/Runtime/NNE/Public/NNERuntimeGPU.h b/Engine/Source/Runtime/NNE/Public/NNERuntimeGPU.h
--- a/Engine/Source/Runtime/NNE/Public/NNERuntimeGPU.h
+++ b/Engine/Source/Runtime/NNE/Public/NNERuntimeGPU.h
@@ -2,8 +2,7 @@
 
 #pragma once
 
-#include "CoreMinimal.h"
-#include "NNERuntime.h"
+#include "Containers/ContainersFwd.h"
 #include "NNETypes.h"
 #include "UObject/Interface.h"
 
diff --git a/Engine/Source/Runtime/NNE/Public/NNERuntimeRDG.h b/Engine/Source/Runtime/NNE/Public/NNERuntimeRDG.h
--- a/Engine/Source/Runtime/NNE/Public/NNERuntimeRDG.h
+++ b/Engine/Source/Runtime/NNE/Public/NNERuntimeRDG.h
@@ -2,7 +2,7 @@
 
 #pragma once
 
-#include "CoreMinimal.h"
+#include "Containers/ContainersFwd.h"
 #include "NNETypes.h"
 #include "RenderGraphFwd.h"
 #include "UObject/Interface.h"
diff --git a/Engine/Source/Runtime/NNE/Public/NNETypes.h b/Engine/Source/Runtime/NNE/Public/NNETypes.h
--- a/Engine/Source/Runtime/NNE/Public/NNETypes.h
+++ b/Engine/Source/Runtime/NNE/Public/NNETypes.h
@@ -2,7 +2,6 @@
 
 #pragma once
 
-#include "CoreMinimal.h"
 #include "UObject/Class.h"
 
 #include "NNETypes.generated.h"
