[NNE] Replaced Logs with checks for detecting contract violations (precondition).

#tests Compile Files (Non-Unity), Unreal Editor (Win64) Style-transfer and smoke tests

#rnx

diff --git a/Engine/Source/Runtime/NNE/Internal/NNEModelBase.h b/Engine/Source/Runtime/NNE/Internal/NNEModelBase.h
--- a/Engine/Source/Runtime/NNE/Internal/NNEModelBase.h
+++ b/Engine/Source/Runtime/NNE/Internal/NNEModelBase.h
@@ -56,7 +56,7 @@ namespace UE::NNE::Internal
 
 		if (InInputShapes.Num() != InputSymbolicTensors.Num())
 		{
-			UE_LOG(LogNNE, Warning, TEXT("Number of input shapes does not match number of input tensors"));
+			UE_LOG(LogNNE, Error, TEXT("Number of input shapes does not match number of input tensors"));
 			return -1;
 		}
 
@@ -65,7 +65,7 @@ namespace UE::NNE::Internal
 			const FTensorDesc SymbolicDesc = InputSymbolicTensors[i];
 			if (!InInputShapes[i].IsCompatibleWith(SymbolicDesc.GetShape()))
 			{
-				UE_LOG(LogNNE, Warning, TEXT("Input shape does not match input tensor %s of index %d"), *SymbolicDesc.GetName(), i);
+				UE_LOG(LogNNE, Error, TEXT("Input shape does not match input tensor %s of index %d"), *SymbolicDesc.GetName(), i);
 				return -1;
 			}
 		}
diff --git a/Engine/Source/Runtime/NNE/Private/NNETypes.cpp b/Engine/Source/Runtime/NNE/Private/NNETypes.cpp
--- a/Engine/Source/Runtime/NNE/Private/NNETypes.cpp
+++ b/Engine/Source/Runtime/NNE/Private/NNETypes.cpp
@@ -52,11 +52,7 @@ void FSymbolicTensorShape::operator=(const FSymbolicTensorShape& OtherShape)
 
 FSymbolicTensorShape FSymbolicTensorShape::Make(TConstArrayView<int32> Data)
 {
-	if (Data.Num() > MaxRank)
-	{
-		UE_LOG(LogNNE, Warning, TEXT("Cannot create symbolic tensor shape, input is rank %d while max rank is %d"), Data.Num(), MaxRank);
-		return {};
-	}
+	checkf(Data.Num() <= MaxRank, TEXT("Cannot create symbolic tensor shape, input is rank %d while max rank is %d"), Data.Num(), MaxRank);
 
 	FSymbolicTensorShape Shape;
 	Shape.Data.Append(Data);
@@ -106,11 +102,7 @@ void FTensorShape::operator=(const FTensorShape& OtherShape)
 
 FTensorShape FTensorShape::Make(TConstArrayView<uint32> Data)
 {
-	if (Data.Num() > MaxRank)
-	{
-		UE_LOG(LogNNE, Warning, TEXT("Cannot create tensor shape, input is rank %d while max rank is %d"), Data.Num(), MaxRank);
-		return {};
-	}
+	checkf(Data.Num() <= MaxRank, TEXT("Cannot create tensor shape, input is rank %d while max rank is %d"), Data.Num(), MaxRank);
 
 	FTensorShape Shape;
 	Shape.Data.Append(Data);
@@ -128,7 +120,7 @@ FTensorShape FTensorShape::MakeFromSymbolic(const FSymbolicTensorShape& Symbolic
 	return ConcreteShape;
 }
 
-size_t GetTensorDataTypeSizeInBytes(ENNETensorDataType InType)
+int32 GetTensorDataTypeSizeInBytes(ENNETensorDataType InType)
 {
 	switch (InType)
 	{
diff --git a/Engine/Source/Runtime/NNE/Public/NNETypes.h b/Engine/Source/Runtime/NNE/Public/NNETypes.h
--- a/Engine/Source/Runtime/NNE/Public/NNETypes.h
+++ b/Engine/Source/Runtime/NNE/Public/NNETypes.h
@@ -41,7 +41,7 @@ namespace UE::NNE
 	 * @param InType the type of the element to consider.
 	 * @return the data size in bytes of an element.
 	 */
-	size_t NNE_API GetTensorDataTypeSizeInBytes(ENNETensorDataType InType);
+	int32 NNE_API GetTensorDataTypeSizeInBytes(ENNETensorDataType InType);
 
 	/**
 	 * A symbolic tensor shape represents the shape of a tensor with potentially variable dimension.
