[NNE] Fix static analyzer warning (PVS-Studio).

diff --git a/Engine/Source/Runtime/NNE/Internal/NNETensor.h b/Engine/Source/Runtime/NNE/Internal/NNETensor.h
--- a/Engine/Source/Runtime/NNE/Internal/NNETensor.h
+++ b/Engine/Source/Runtime/NNE/Internal/NNETensor.h
@@ -57,7 +57,7 @@ namespace UE::NNE::Internal
 			check(InShape.Volume() <= TNumericLimits<uint32>::Max());
 			Shape = InShape;
 			Volume = InShape.Volume();
-			DataSize = static_cast<uint64>(GetTensorDataTypeSizeInBytes(DataType) * Volume);
+			DataSize = (uint64)GetTensorDataTypeSizeInBytes(DataType) * Volume;
 		}
 
 		template <typename T> void SetPreparedData(TConstArrayView<T> Data)
