[NNE] Runtime registry cleanup:
- Register/Unregister runtime API calls return status code instead of bool
- Internally use TMap instead of TArray
- Replaced GetAllRuntimes() with GetAllRuntimeNames() (including templated variant to get names of runtimes that implement the specified interface)

#tests Editor/Game (Win64) style-transfer and smoke tests.

#rnx

diff --git a/Engine/Plugins/Experimental/NNEDenoiser/Source/NNEDenoiser/Private/NNEDenoiserModelInstanceCPU.cpp b/Engine/Plugins/Experimental/NNEDenoiser/Source/NNEDenoiser/Private/NNEDenoiserModelInstanceCPU.cpp
--- a/Engine/Plugins/Experimental/NNEDenoiser/Source/NNEDenoiser/Private/NNEDenoiserModelInstanceCPU.cpp
+++ b/Engine/Plugins/Experimental/NNEDenoiser/Source/NNEDenoiser/Private/NNEDenoiserModelInstanceCPU.cpp
@@ -31,12 +31,9 @@ namespace UE::NNEDenoiser::Private
 		if (!RuntimeCPU.IsValid())
 		{
 			UE_LOG(LogNNEDenoiser, Error, TEXT("No CPU runtime '%s' found. Valid CPU runtimes are: "), *RuntimeCPUName);
-			for (TWeakInterfacePtr<INNERuntime> Runtime : UE::NNE::GetAllRuntimes())
+			for (const FString& RuntimeName : UE::NNE::GetAllRuntimeNames<INNERuntimeCPU>())
 			{
-				if (Runtime.IsValid() && UE::NNE::GetRuntime<INNERuntimeCPU>(Runtime->GetRuntimeName()).IsValid())
-				{
-					UE_LOG(LogNNEDenoiser, Error, TEXT("- %s"), *Runtime->GetRuntimeName());
-				}
+				UE_LOG(LogNNEDenoiser, Error, TEXT("- %s"), *RuntimeName);
 			}
 			return {};
 		}
diff --git a/Engine/Plugins/Experimental/NNEDenoiser/Source/NNEDenoiser/Private/NNEDenoiserModelInstanceGPU.cpp b/Engine/Plugins/Experimental/NNEDenoiser/Source/NNEDenoiser/Private/NNEDenoiserModelInstanceGPU.cpp
--- a/Engine/Plugins/Experimental/NNEDenoiser/Source/NNEDenoiser/Private/NNEDenoiserModelInstanceGPU.cpp
+++ b/Engine/Plugins/Experimental/NNEDenoiser/Source/NNEDenoiser/Private/NNEDenoiserModelInstanceGPU.cpp
@@ -34,12 +34,9 @@ namespace UE::NNEDenoiser::Private
 		if (!RuntimeGPU.IsValid())
 		{
 			UE_LOG(LogNNEDenoiser, Error, TEXT("No GPU runtime '%s' found. Valid GPU runtimes are: "), *RuntimeGPUName);
-			for (TWeakInterfacePtr<INNERuntime> Runtime : UE::NNE::GetAllRuntimes())
+			for (const FString& RuntimeName : UE::NNE::GetAllRuntimeNames<INNERuntimeGPU>())
 			{
-				if (Runtime.IsValid() && UE::NNE::GetRuntime<INNERuntimeGPU>(Runtime->GetRuntimeName()).IsValid())
-				{
-					UE_LOG(LogNNEDenoiser, Error, TEXT("- %s"), *Runtime->GetRuntimeName());
-				}
+				UE_LOG(LogNNEDenoiser, Error, TEXT("- %s"), *RuntimeName);
 			}
 			return {};
 		}
diff --git a/Engine/Plugins/Experimental/NNEDenoiser/Source/NNEDenoiser/Private/NNEDenoiserModelInstanceRDG.cpp b/Engine/Plugins/Experimental/NNEDenoiser/Source/NNEDenoiser/Private/NNEDenoiserModelInstanceRDG.cpp
--- a/Engine/Plugins/Experimental/NNEDenoiser/Source/NNEDenoiser/Private/NNEDenoiserModelInstanceRDG.cpp
+++ b/Engine/Plugins/Experimental/NNEDenoiser/Source/NNEDenoiser/Private/NNEDenoiserModelInstanceRDG.cpp
@@ -27,12 +27,9 @@ namespace UE::NNEDenoiser::Private
 		if (!RuntimeRDG.IsValid())
 		{
 			UE_LOG(LogNNEDenoiser, Error, TEXT("No RDG runtime '%s' found. Valid RDG runtimes are: "), *RuntimeRDGName);
-			for (TWeakInterfacePtr<INNERuntime> Runtime : UE::NNE::GetAllRuntimes())
+			for (const FString& RuntimeName : UE::NNE::GetAllRuntimeNames<INNERuntimeRDG>())
 			{
-				if (Runtime.IsValid() && UE::NNE::GetRuntime<INNERuntimeRDG>(Runtime->GetRuntimeName()).IsValid())
-				{
-					UE_LOG(LogNNEDenoiser, Error, TEXT("- %s"), *Runtime->GetRuntimeName());
-				}
+				UE_LOG(LogNNEDenoiser, Error, TEXT("- %s"), *RuntimeName);
 			}
 			return {};
 		}
diff --git a/Engine/Source/Editor/NNEEditor/Private/NNEEditorRuntimesWidget.cpp b/Engine/Source/Editor/NNEEditor/Private/NNEEditorRuntimesWidget.cpp
--- a/Engine/Source/Editor/NNEEditor/Private/NNEEditorRuntimesWidget.cpp
+++ b/Engine/Source/Editor/NNEEditor/Private/NNEEditorRuntimesWidget.cpp
@@ -125,11 +125,7 @@ void SRuntimesWidget::InitCheckBoxNames()
 	CheckBoxNames.Reset();
 
 	CheckBoxNames.Add(TEXT("All"));
-
-	for (TWeakInterfacePtr<INNERuntime> Runtime : UE::NNE::GetAllRuntimes())
-	{
-		CheckBoxNames.Add(Runtime->GetRuntimeName());
-	}
+	CheckBoxNames.Append(UE::NNE::GetAllRuntimeNames());
 
 	//A runtime could possibly be in the target list but not registered to the editor
 	for (const FString& EnabledRuntimeName : TargetRuntimes.Get())
diff --git a/Engine/Source/Runtime/NNE/Private/NNE.cpp b/Engine/Source/Runtime/NNE/Private/NNE.cpp
--- a/Engine/Source/Runtime/NNE/Private/NNE.cpp
+++ b/Engine/Source/Runtime/NNE/Private/NNE.cpp
@@ -11,95 +11,88 @@ DEFINE_LOG_CATEGORY(LogNNE);
 
 namespace UE::NNE
 {
+
 	class FRegistry
 	{
 	public:
-
-		static FRegistry* Get()
+		static FRegistry& GetInstance()
 		{
-			static FRegistry Inst;
-
-			return &Inst;
+			static FRegistry Instance;
+			return Instance;
 		}
 
-		bool Add(TWeakInterfacePtr<INNERuntime> Runtime)
+		EResultStatus Add(TWeakInterfacePtr<INNERuntime> Runtime)
 		{
-			if (!Runtime.IsValid())
-			{
-				return false;
-			}
-			
-			if (FindByName(Runtime->GetRuntimeName()) >= 0)
+			checkf(Runtime.IsValid(), TEXT("Runtime is not valid"));
+
+			const FString RuntimeName = Runtime->GetRuntimeName();
+			checkf(!RuntimeName.IsEmpty(), TEXT("Runtime name is empty"));
+
+			if (Runtimes.Contains(RuntimeName))
 			{
-				UE_LOG(LogNNE, Warning, TEXT("Runtime %s is already registered"), *Runtime->GetRuntimeName());
-				return false;
+				UE_LOG(LogNNE, Warning, TEXT("Runtime %s is already registered"), *RuntimeName);
+				return EResultStatus::Fail;
 			}
 
-			Runtimes.Add(Runtime);
+			Runtimes.Add(RuntimeName, Runtime);
 
-			return true;
+			return EResultStatus::Ok;
 		}
 
-		bool Remove(TWeakInterfacePtr<INNERuntime> Runtime)
+		EResultStatus Remove(TWeakInterfacePtr<INNERuntime> Runtime)
 		{
-			if (!Runtime.IsValid())
-			{
-				return false;
-			}
-			
-			int Index = FindByName(Runtime->GetRuntimeName());
-			if (Index >= 0)
-			{
-				Runtimes.RemoveAtSwap(Index);
-				return true;
-			}
+			checkf(Runtime.IsValid(), TEXT("Runtime is not valid"));
 
-			return false;
-		}
+			const FString RuntimeName = Runtime->GetRuntimeName();
+			checkf(!RuntimeName.IsEmpty(), TEXT("Runtime name is empty"));
 
+			return Runtimes.Remove(RuntimeName) >= 1 ? EResultStatus::Ok : EResultStatus::Fail;
+		}
 
-		TArrayView<TWeakInterfacePtr<INNERuntime>> GetAllRuntimes()
+		TWeakInterfacePtr<INNERuntime> Get(const FString& Name) const
 		{
-			RemoveInvalidPtrs();
-			return Runtimes;
-		}
+			checkf(!Name.IsEmpty(), TEXT("Name is empty"));
 
-	private:
+			if (Runtimes.Contains(Name))
+			{
+				TWeakInterfacePtr<INNERuntime> Result = Runtimes.FindChecked(Name);
+				ensureMsgf(Result.IsValid(), TEXT("Runtime %s is not valid"), *Name);
 
-		void RemoveInvalidPtrs()
-		{
-			Runtimes.RemoveAllSwap([](TWeakInterfacePtr<INNERuntime> Runtime) -> bool {return !Runtime.IsValid();}, false);
+				return Result;
+			}
+
+			return {};
 		}
 
-		int FindByName(const FString& Name)
+		TArray<FString> GetAllNames() const
 		{
-			RemoveInvalidPtrs();
-			
-			int Index = -1;
+			TArray<FString> Result;
+			Runtimes.GenerateKeyArray(Result);
 
-			for (int Idx = 0; Idx < Runtimes.Num(); ++Idx)
+			Result.SetNum(Algo::RemoveIf(Result, [&] (const FString &RuntimeName)
 			{
-				if (Runtimes[Idx]->GetRuntimeName() == Name)
-				{
-					Index = Idx;
-					break;
-				}
-			}
+				return !ensureMsgf(Runtimes.FindChecked(RuntimeName).IsValid(), TEXT("Runtime %s is not valid"), *RuntimeName);
+			}));
 
-			return Index;
+			return Result;
 		}
 
-		TArray<TWeakInterfacePtr<INNERuntime>>	Runtimes;
+	private:
+		TMap<FString, TWeakInterfacePtr<INNERuntime>> Runtimes;
 	};
 
-	bool RegisterRuntime(TWeakInterfacePtr<INNERuntime> Runtime)
+	ERegisterRuntimeResultStatus RegisterRuntime(TWeakInterfacePtr<INNERuntime> Runtime)
 	{
+		const ERegisterRuntimeResultStatus Result = FRegistry::GetInstance().Add(Runtime);
+
 #ifdef WITH_EDITOR
 		FModuleManager::Get().LoadModule(TEXT("NNEEditor"));
 #endif
 
-		FString RuntimeName = Runtime->GetRuntimeName();
-		FCoreDelegates::OnAllModuleLoadingPhasesComplete.AddLambda([RuntimeName]() {
+		const FString RuntimeName = Runtime->GetRuntimeName();
+
+		FCoreDelegates::OnAllModuleLoadingPhasesComplete.AddLambda([RuntimeName]()
+		{
 			if (FEngineAnalytics::IsAvailable())
 			{
 				TArray<FAnalyticsEventAttribute> Attributes = MakeAnalyticsEventAttributeArray(
@@ -110,16 +103,21 @@ namespace UE::NNE
 			}
 		});
 
-		return FRegistry::Get()->Add(Runtime);
+		return Result;
+	}
+
+	EUnregisterRuntimeResultStatus UnregisterRuntime(TWeakInterfacePtr<INNERuntime> Runtime)
+	{
+		return FRegistry::GetInstance().Remove(Runtime);
 	}
 
-	bool UnregisterRuntime(TWeakInterfacePtr<INNERuntime> Runtime)
+	TArray<FString> GetAllRuntimeNames()
 	{
-		return FRegistry::Get()->Remove(Runtime);
+		return FRegistry::GetInstance().GetAllNames();
 	}
 
-	TArrayView<TWeakInterfacePtr<INNERuntime>> GetAllRuntimes()
+	TWeakInterfacePtr<INNERuntime> GetRuntime(const FString& Name)
 	{
-		return FRegistry::Get()->GetAllRuntimes();
+		return FRegistry::GetInstance().Get(Name);
 	}
 }
diff --git a/Engine/Source/Runtime/NNE/Private/NNEModelData.cpp b/Engine/Source/Runtime/NNE/Private/NNEModelData.cpp
--- a/Engine/Source/Runtime/NNE/Private/NNEModelData.cpp
+++ b/Engine/Source/Runtime/NNE/Private/NNEModelData.cpp
@@ -144,10 +144,10 @@ namespace UE::NNE::ModelData
 		else
 		{
 			UE_LOG(LogNNE, Error, TEXT("UNNEModelData: No runtime '%s' found. Valid runtimes are: "), *RuntimeName);
-			TArrayView<TWeakInterfacePtr<INNERuntime>> Runtimes = UE::NNE::GetAllRuntimes();
+			TArray<FString> Runtimes = UE::NNE::GetAllRuntimeNames();
 			for (int32 i = 0; i < Runtimes.Num(); i++)
 			{
-				UE_LOG(LogNNE, Error, TEXT("- %s"), *Runtimes[i]->GetRuntimeName());
+				UE_LOG(LogNNE, Error, TEXT("- %s"), *Runtimes[i]);
 			}
 		}
 		return TSharedPtr<UE::NNE::FSharedModelData>();
@@ -240,10 +240,7 @@ void UNNEModelData::Serialize(FArchive& Ar)
 				TArray<FString, TInlineAllocator<10>> CookRuntimeNames;
 				if (GetTargetRuntimes().IsEmpty())
 				{
-					for (const TWeakInterfacePtr<INNERuntime>& Runtime : UE::NNE::GetAllRuntimes())
-					{
-						CookRuntimeNames.Add(Runtime->GetRuntimeName());
-					}
+					CookRuntimeNames.Append(UE::NNE::GetAllRuntimeNames());
 				}
 				else
 				{
diff --git a/Engine/Source/Runtime/NNE/Public/NNE.h b/Engine/Source/Runtime/NNE/Public/NNE.h
--- a/Engine/Source/Runtime/NNE/Public/NNE.h
+++ b/Engine/Source/Runtime/NNE/Public/NNE.h
@@ -2,39 +2,54 @@
 
 #pragma once
 
+#include "Algo/RemoveIf.h"
 #include "Containers/ArrayView.h"
 #include "Containers/UnrealString.h"
 #include "NNERuntime.h"
+#include "NNEStatus.h"
 #include "UObject/WeakInterfacePtr.h"
 
 NNE_API DECLARE_LOG_CATEGORY_EXTERN(LogNNE, Log, All);
 
 namespace UE::NNE
 {
+	using ERegisterRuntimeResultStatus = EResultStatus;
+	using EUnregisterRuntimeResultStatus = EResultStatus;
+
 	/**
 	 * Register a runtime to make it accessible to NNE clients.
 	 *
 	 * The caller needs to keep a strong pointer to the runtime object to prevent it from being garbage collected.
 	 *
 	 * @param Runtime A weak interface pointer to the runtime to be registered.
-	 * @return True if the runtime has been registered successfully, false otherwise.
+	 * @return Status indicating success or failure (e.g. if the runtime already has been registered).
 	 */
-	NNE_API bool RegisterRuntime(TWeakInterfacePtr<INNERuntime> Runtime);
+	NNE_API ERegisterRuntimeResultStatus RegisterRuntime(TWeakInterfacePtr<INNERuntime> Runtime);
 
 	/**
 	 * Unregister a registered runtime.
 	 *
 	 * @param Runtime A weak interface pointer to the runtime to be unregistered.
-	 * @return True if the runtime has been unregistered successfully, false otherwise (e.g. if the runtime has not been registered).
+	 * @return Status indicating success or failure (e.g. if the runtime has not been registered).
 	 */
-	NNE_API bool UnregisterRuntime(TWeakInterfacePtr<INNERuntime> Runtime);
+	NNE_API EUnregisterRuntimeResultStatus UnregisterRuntime(TWeakInterfacePtr<INNERuntime> Runtime);
 	
 	/**
-	 * List and return all registered runtimes.
+	 * List and return all registered runtime names.
 	 *
-	 * @return An array containing weak pointers to all registered runtimes.
+	 * @return An array containing runtime names of all registered runtimes.
 	 */
-	NNE_API TArrayView<TWeakInterfacePtr<INNERuntime>> GetAllRuntimes();
+	NNE_API TArray<FString> GetAllRuntimeNames();
+
+	/**
+	 * Find and return a runtime by name.
+	 *
+	 * This function tries to find a runtime by name.
+	 *
+	 * @param Name The name of the runtime.
+	 * @return A weak pointer to the runtime if it has been found or an invalid pointer otherwise.
+	 */
+	NNE_API TWeakInterfacePtr<INNERuntime> GetRuntime(const FString& Name);
 
 	/**
 	 * Find and return a runtime by name and interface.
@@ -44,17 +59,30 @@ namespace UE::NNE
 	 * @param Name The name of the runtime.
 	 * @return A weak pointer to the runtime if it has been found and implements the interface in the template argument or an invalid pointer otherwise.
 	 */
-	template<class T> TWeakInterfacePtr<T> GetRuntime(const FString& Name)
+	template<class T>
+	TWeakInterfacePtr<T> GetRuntime(const FString& Name)
 	{
-		for (TWeakInterfacePtr<INNERuntime> Runtime : GetAllRuntimes())
+		TWeakInterfacePtr<INNERuntime> Runtime = GetRuntime(Name);
+
+		T* RuntimePtr = Cast<T>(Runtime.Get());
+
+		return TWeakInterfacePtr<T>(RuntimePtr);
+	}
+
+	/**
+	 * List and return all registered runtime names that implement the provided interface.
+	 *
+	 * @return An array containing runtime names of all registered runtimes that implement the interface in the template argument.
+	 */
+	template<class T>
+	TArray<FString> GetAllRuntimeNames()
+	{
+		TArray<FString> Result = GetAllRuntimeNames();
+		Result.SetNum(Algo::RemoveIf(Result, [] (const FString &RuntimeName)
 		{
-			if (Runtime->GetRuntimeName() == Name)
-			{
-				T* RuntimePtr = Cast<T>(Runtime.Get());
-				return TWeakInterfacePtr<T>(RuntimePtr);
-			}
-		}
-		return TWeakInterfacePtr<T>(nullptr);
+			return GetRuntime<T>(RuntimeName).IsValid();
+		}));
+
+		return Result;
 	}
-	
 }
diff --git a/Engine/Source/Runtime/NNE/Public/NNERuntime.h b/Engine/Source/Runtime/NNE/Public/NNERuntime.h
--- a/Engine/Source/Runtime/NNE/Public/NNERuntime.h
+++ b/Engine/Source/Runtime/NNE/Public/NNERuntime.h
@@ -22,7 +22,7 @@ class NNE_API UNNERuntime : public UInterface
 };
 
 /**
- * The base interface of a neural network runtime. It is returned by UE::NNE::GetAllRuntimes() giving access to all available runtimes.
+ * The base interface of a neural network runtime.
  *
  * This interface is mainly used internally by UNNEModelData to cook a model given in a file format (e.g. .onnx) into runtime specific model data.
  * The model data is then stored inside UNNEModelData which then can be used in game to create an inferable model.
