[NNE] Deprecate nne.Editor.UseDDCForCooking

Problem:
nne.Editor.UseDDCForCooking is incompatible with runtime producing custom cooking results not captured by SharedModelData.

Solution:
make the setting deprecated.

Test:
* Win64 Editor: Cook with and without the UseDDCForCooking, verify correct DDC behavior + double check cvar is not exposed in the UI anymore.
* Win64 Standalone: run a subset of tests

#rnx

diff --git a/Engine/Source/Editor/NNEEditor/Private/NNEEditorSettings.h b/Engine/Source/Editor/NNEEditor/Private/NNEEditorSettings.h
--- a/Engine/Source/Editor/NNEEditor/Private/NNEEditorSettings.h
+++ b/Engine/Source/Editor/NNEEditor/Private/NNEEditorSettings.h
@@ -24,10 +24,4 @@ public:
 		ConsoleVariable = "nne.Editor.UseDDC", DisplayName = "Use DDC for Neural Network optimizations",
 		ToolTip = "True if NNE should use the DDC cache to fetch previously optimized Neural Network."))
 	bool bUseDDCCache;
-
-	UPROPERTY(config, EditAnywhere, BlueprintReadWrite, Category = Editor, meta = (
-		EditCondition = "bUseDDCCache",
-		ConsoleVariable = "nne.Editor.UseDDCForCooking", DisplayName = "Use DDC for Neural Network optimizations while cooking",
-		ToolTip = "True if NNE should use the DDC cache to fetch previously optimized Neural Network while cooking."))
-	bool bUseDDCCacheForCooking;
 };
\ No newline at end of file
diff --git a/Engine/Source/Runtime/NNE/Private/NNEModelData.cpp b/Engine/Source/Runtime/NNE/Private/NNEModelData.cpp
--- a/Engine/Source/Runtime/NNE/Private/NNEModelData.cpp
+++ b/Engine/Source/Runtime/NNE/Private/NNEModelData.cpp
@@ -184,13 +184,13 @@ namespace UE::NNE
 		static TAutoConsoleVariable<bool> CVarNNEEditorUseDDC(
 			TEXT("nne.Editor.UseDDC"),
 			true,
-			TEXT("Indicates whether NNE wether NNE should use the DDC cache to speed up SharedModelData creation in the editor, see also nne.Editor.UseDDCForCooking."),
+			TEXT("Indicates whether NNE should use the DDC cache to speed up SharedModelData creation in the editor, see also nne.Editor.UseDDCForCooking."),
 			ECVF_Default);
 
-		static TAutoConsoleVariable<bool> CVarNNEEditorUseDDCForCooking(
-			TEXT("nne.Editor.UseDDCForCooking"),
+		static TAutoConsoleVariable<bool> CVarNNEEditorUseDDCForCookingDeprecated(
+			TEXT("nne.Editor.UseDDCForCooking_Deprecated"),
 			false,
-			TEXT("Indicates whether NNE wether NNE should use the DDC cache to speed up SharedModelData creation while cooking, if nne.Editor.UseDDC is false cooking won't use the DDC even if nne.Editor.UseDDCForCooking is true."),
+			TEXT("DEPRECATED. Indicates whether NNE should use the DDC cache to speed up SharedModelData creation while cooking, if nne.Editor.UseDDC is false cooking won't use the DDC even if nne.Editor.UseDDCForCooking is true."),
 			ECVF_Default);
 	} // ConsoleVariables
 #endif //WITH_EDITOR
@@ -267,7 +267,7 @@ void UNNEModelData::Serialize(FArchive& Ar)
 						if (ModelDataIdentifier.Len() > 0)
 						{
 							if (UE::NNE::ConsoleVariables::CVarNNEEditorUseDDC.GetValueOnGameThread() &&
-								UE::NNE::ConsoleVariables::CVarNNEEditorUseDDCForCooking.GetValueOnGameThread())
+								UE::NNE::ConsoleVariables::CVarNNEEditorUseDDCForCookingDeprecated.GetValueOnGameThread())
 							{
 								SharedModelData = UE::NNE::ModelData::GetFromDDC(FileId, RuntimeName, ModelDataIdentifier);
 							}
