[NNE] Uniformising enum status name to match called function

#tests win64 editor: smoke test

#rnx

diff --git a/Engine/Plugins/Experimental/NNEDenoiser/Source/NNEDenoiser/Private/NNEDenoiserIOProcessBase.cpp b/Engine/Plugins/Experimental/NNEDenoiser/Source/NNEDenoiser/Private/NNEDenoiserIOProcessBase.cpp
--- a/Engine/Plugins/Experimental/NNEDenoiser/Source/NNEDenoiser/Private/NNEDenoiserIOProcessBase.cpp
+++ b/Engine/Plugins/Experimental/NNEDenoiser/Source/NNEDenoiser/Private/NNEDenoiserIOProcessBase.cpp
@@ -189,8 +189,8 @@ bool FInputProcessBase::PrepareAndValidate(IModelInstance& ModelInstance, FIntPo
 		UE_LOG(LogNNEDenoiser, Log, TEXT("%d: (%d, %d, %d, %d)"), Idx, ModelInputShape.X, ModelInputShape.Y, ModelInputShape.Z, ModelInputShape.W)
 	}
 
-	const IModelInstance::ESetInputTensorShapeStatus Status = ModelInstance.SetInputTensorShapes(InputShapes);
-	if (Status != IModelInstance::ESetInputTensorShapeStatus::Ok)
+	const IModelInstance::ESetInputTensorShapesStatus Status = ModelInstance.SetInputTensorShapes(InputShapes);
+	if (Status != IModelInstance::ESetInputTensorShapesStatus::Ok)
 	{
 		UE_LOG(LogNNEDenoiser, Error, TEXT("Could not configure model instance (ModelInstance.SetInputTensorShapes() failed)!"))
 		return false;
diff --git a/Engine/Plugins/Experimental/NNEDenoiser/Source/NNEDenoiser/Private/NNEDenoiserModelInstance.h b/Engine/Plugins/Experimental/NNEDenoiser/Source/NNEDenoiser/Private/NNEDenoiserModelInstance.h
--- a/Engine/Plugins/Experimental/NNEDenoiser/Source/NNEDenoiser/Private/NNEDenoiserModelInstance.h
+++ b/Engine/Plugins/Experimental/NNEDenoiser/Source/NNEDenoiser/Private/NNEDenoiserModelInstance.h
@@ -10,7 +10,7 @@ namespace UE::NNEDenoiser::Private
 class IModelInstance : public NNE::IModelInstanceRDG
 {
 public:
-	using ESetInputTensorShapeStatus = NNE::IModelInstanceRDG::ESetInputTensorShapeStatus;
+	using ESetInputTensorShapesStatus = NNE::IModelInstanceRDG::ESetInputTensorShapesStatus;
 	using EEnqueueRDGStatus = NNE::IModelInstanceRDG::EEnqueueRDGStatus;
 
 	virtual ~IModelInstance() = default;
diff --git a/Engine/Plugins/Experimental/NNEDenoiser/Source/NNEDenoiser/Private/NNEDenoiserModelInstanceCPU.cpp b/Engine/Plugins/Experimental/NNEDenoiser/Source/NNEDenoiser/Private/NNEDenoiserModelInstanceCPU.cpp
--- a/Engine/Plugins/Experimental/NNEDenoiser/Source/NNEDenoiser/Private/NNEDenoiserModelInstanceCPU.cpp
+++ b/Engine/Plugins/Experimental/NNEDenoiser/Source/NNEDenoiser/Private/NNEDenoiserModelInstanceCPU.cpp
@@ -88,11 +88,11 @@ namespace UE::NNEDenoiser::Private
 		return ModelInstance->GetOutputTensorShapes();
 	}
 
-	FModelInstanceCPU::ESetInputTensorShapeStatus FModelInstanceCPU::SetInputTensorShapes(TConstArrayView<NNE::FTensorShape> InInputShapes)
+	FModelInstanceCPU::ESetInputTensorShapesStatus FModelInstanceCPU::SetInputTensorShapes(TConstArrayView<NNE::FTensorShape> InInputShapes)
 	{
-		NNE::IModelInstanceCPU::ESetInputTensorShapeStatus Status = ModelInstance->SetInputTensorShapes(InInputShapes);
+		NNE::IModelInstanceCPU::ESetInputTensorShapesStatus Status = ModelInstance->SetInputTensorShapes(InInputShapes);
 
-		return Status == NNE::IModelInstanceCPU::ESetInputTensorShapeStatus::Ok ? ESetInputTensorShapeStatus::Ok : ESetInputTensorShapeStatus::Fail;
+		return Status == NNE::IModelInstanceCPU::ESetInputTensorShapesStatus::Ok ? ESetInputTensorShapesStatus::Ok : ESetInputTensorShapesStatus::Fail;
 	}
 
 	FModelInstanceCPU::EEnqueueRDGStatus FModelInstanceCPU::EnqueueRDG(FRDGBuilder &GraphBuilder, TConstArrayView<NNE::FTensorBindingRDG> Inputs, TConstArrayView<NNE::FTensorBindingRDG> Outputs)
diff --git a/Engine/Plugins/Experimental/NNEDenoiser/Source/NNEDenoiser/Private/NNEDenoiserModelInstanceCPU.h b/Engine/Plugins/Experimental/NNEDenoiser/Source/NNEDenoiser/Private/NNEDenoiserModelInstanceCPU.h
--- a/Engine/Plugins/Experimental/NNEDenoiser/Source/NNEDenoiser/Private/NNEDenoiserModelInstanceCPU.h
+++ b/Engine/Plugins/Experimental/NNEDenoiser/Source/NNEDenoiser/Private/NNEDenoiserModelInstanceCPU.h
@@ -32,7 +32,7 @@ public:
 
 	virtual TConstArrayView<NNE::FTensorShape> GetOutputTensorShapes() const override;
 
-	virtual ESetInputTensorShapeStatus SetInputTensorShapes(TConstArrayView<NNE::FTensorShape> InInputShapes) override;
+	virtual ESetInputTensorShapesStatus SetInputTensorShapes(TConstArrayView<NNE::FTensorShape> InInputShapes) override;
 
 	virtual EEnqueueRDGStatus EnqueueRDG(FRDGBuilder &GraphBuilder, TConstArrayView<NNE::FTensorBindingRDG> Inputs, TConstArrayView<NNE::FTensorBindingRDG> Outputs) override;
 
diff --git a/Engine/Plugins/Experimental/NNEDenoiser/Source/NNEDenoiser/Private/NNEDenoiserModelInstanceGPU.cpp b/Engine/Plugins/Experimental/NNEDenoiser/Source/NNEDenoiser/Private/NNEDenoiserModelInstanceGPU.cpp
--- a/Engine/Plugins/Experimental/NNEDenoiser/Source/NNEDenoiser/Private/NNEDenoiserModelInstanceGPU.cpp
+++ b/Engine/Plugins/Experimental/NNEDenoiser/Source/NNEDenoiser/Private/NNEDenoiserModelInstanceGPU.cpp
@@ -91,11 +91,11 @@ namespace UE::NNEDenoiser::Private
 		return ModelInstance->GetOutputTensorShapes();
 	}
 
-	FModelInstanceGPU::ESetInputTensorShapeStatus FModelInstanceGPU::SetInputTensorShapes(TConstArrayView<NNE::FTensorShape> InInputShapes)
+	FModelInstanceGPU::ESetInputTensorShapesStatus FModelInstanceGPU::SetInputTensorShapes(TConstArrayView<NNE::FTensorShape> InInputShapes)
 	{
-		NNE::IModelInstanceGPU::ESetInputTensorShapeStatus Status = ModelInstance->SetInputTensorShapes(InInputShapes);
+		NNE::IModelInstanceGPU::ESetInputTensorShapesStatus Status = ModelInstance->SetInputTensorShapes(InInputShapes);
 
-		return Status == NNE::IModelInstanceGPU::ESetInputTensorShapeStatus::Ok ? ESetInputTensorShapeStatus::Ok : ESetInputTensorShapeStatus::Fail;
+		return Status == NNE::IModelInstanceGPU::ESetInputTensorShapesStatus::Ok ? ESetInputTensorShapesStatus::Ok : ESetInputTensorShapesStatus::Fail;
 	}
 
 	FModelInstanceGPU::EEnqueueRDGStatus FModelInstanceGPU::EnqueueRDG(FRDGBuilder &GraphBuilder, TConstArrayView<NNE::FTensorBindingRDG> Inputs, TConstArrayView<NNE::FTensorBindingRDG> Outputs)
diff --git a/Engine/Plugins/Experimental/NNEDenoiser/Source/NNEDenoiser/Private/NNEDenoiserModelInstanceGPU.h b/Engine/Plugins/Experimental/NNEDenoiser/Source/NNEDenoiser/Private/NNEDenoiserModelInstanceGPU.h
--- a/Engine/Plugins/Experimental/NNEDenoiser/Source/NNEDenoiser/Private/NNEDenoiserModelInstanceGPU.h
+++ b/Engine/Plugins/Experimental/NNEDenoiser/Source/NNEDenoiser/Private/NNEDenoiserModelInstanceGPU.h
@@ -30,7 +30,7 @@ public:
 
 	virtual TConstArrayView<NNE::FTensorShape> GetOutputTensorShapes() const override;
 
-	virtual ESetInputTensorShapeStatus SetInputTensorShapes(TConstArrayView<NNE::FTensorShape> InInputShapes) override;
+	virtual ESetInputTensorShapesStatus SetInputTensorShapes(TConstArrayView<NNE::FTensorShape> InInputShapes) override;
 
 	virtual EEnqueueRDGStatus EnqueueRDG(FRDGBuilder &GraphBuilder, TConstArrayView<NNE::FTensorBindingRDG> Inputs, TConstArrayView<NNE::FTensorBindingRDG> Outputs) override;
 
diff --git a/Engine/Plugins/Experimental/NNEDenoiser/Source/NNEDenoiser/Private/NNEDenoiserModelInstanceRDG.cpp b/Engine/Plugins/Experimental/NNEDenoiser/Source/NNEDenoiser/Private/NNEDenoiserModelInstanceRDG.cpp
--- a/Engine/Plugins/Experimental/NNEDenoiser/Source/NNEDenoiser/Private/NNEDenoiserModelInstanceRDG.cpp
+++ b/Engine/Plugins/Experimental/NNEDenoiser/Source/NNEDenoiser/Private/NNEDenoiserModelInstanceRDG.cpp
@@ -84,7 +84,7 @@ namespace UE::NNEDenoiser::Private
 		return ModelInstance->GetOutputTensorShapes();
 	}
 
-	FModelInstanceRDG::ESetInputTensorShapeStatus FModelInstanceRDG::SetInputTensorShapes(TConstArrayView<NNE::FTensorShape> InInputShapes)
+	FModelInstanceRDG::ESetInputTensorShapesStatus FModelInstanceRDG::SetInputTensorShapes(TConstArrayView<NNE::FTensorShape> InInputShapes)
 	{
 		return ModelInstance->SetInputTensorShapes(InInputShapes);
 	}
diff --git a/Engine/Plugins/Experimental/NNEDenoiser/Source/NNEDenoiser/Private/NNEDenoiserModelInstanceRDG.h b/Engine/Plugins/Experimental/NNEDenoiser/Source/NNEDenoiser/Private/NNEDenoiserModelInstanceRDG.h
--- a/Engine/Plugins/Experimental/NNEDenoiser/Source/NNEDenoiser/Private/NNEDenoiserModelInstanceRDG.h
+++ b/Engine/Plugins/Experimental/NNEDenoiser/Source/NNEDenoiser/Private/NNEDenoiserModelInstanceRDG.h
@@ -30,7 +30,7 @@ public:
 
 	virtual TConstArrayView<NNE::FTensorShape> GetOutputTensorShapes() const override;
 
-	virtual ESetInputTensorShapeStatus SetInputTensorShapes(TConstArrayView<NNE::FTensorShape> InInputShapes) override;
+	virtual ESetInputTensorShapesStatus SetInputTensorShapes(TConstArrayView<NNE::FTensorShape> InInputShapes) override;
 
 	virtual EEnqueueRDGStatus EnqueueRDG(FRDGBuilder &GraphBuilder, TConstArrayView<NNE::FTensorBindingRDG> Inputs, TConstArrayView<NNE::FTensorBindingRDG> Outputs) override;
 
diff --git a/Engine/Plugins/Experimental/NNERuntimeBasicCpu/Source/NNERuntimeBasicCpu/Private/NNERuntimeBasicCpuModel.cpp b/Engine/Plugins/Experimental/NNERuntimeBasicCpu/Source/NNERuntimeBasicCpu/Private/NNERuntimeBasicCpuModel.cpp
--- a/Engine/Plugins/Experimental/NNERuntimeBasicCpu/Source/NNERuntimeBasicCpu/Private/NNERuntimeBasicCpuModel.cpp
+++ b/Engine/Plugins/Experimental/NNERuntimeBasicCpu/Source/NNERuntimeBasicCpu/Private/NNERuntimeBasicCpuModel.cpp
@@ -3383,20 +3383,20 @@ namespace UE::NNE::RuntimeBasic
 		, Instance(Model->Layer->MakeInstance())
 	{}
 
-	FModelInstanceCPU::ESetInputTensorShapeStatus FModelInstanceCPU::SetInputTensorShapes(TConstArrayView<FTensorShape> InInputShapes)
+	FModelInstanceCPU::ESetInputTensorShapesStatus FModelInstanceCPU::SetInputTensorShapes(TConstArrayView<FTensorShape> InInputShapes)
 	{
 		NNE_RUNTIME_BASIC_TRACE_SCOPE(NNE::RuntimeBasic::FModelInstanceCPU::SetInputTensorShapes);
 
 		if (!ensureMsgf(InInputShapes.Num() == 1, TEXT("Basic CPU Inference only supports single input tensor.")))
 		{
-			return ESetInputTensorShapeStatus::Fail;
+			return ESetInputTensorShapesStatus::Fail;
 		}
 
 		const FTensorShape& InputShape = InInputShapes[0];
 
 		if (!ensureMsgf(InputShape.Rank() == 2, TEXT("Basic CPU Inference only supports rank 2 input tensors.")))
 		{
-			return ESetInputTensorShapeStatus::Fail;
+			return ESetInputTensorShapesStatus::Fail;
 		}
 
 		const uint32 InputBatchSize = InputShape.GetData()[0];
@@ -3406,7 +3406,7 @@ namespace UE::NNE::RuntimeBasic
 
 		if (!ensureMsgf(InputInputSize == ModelInputSize, TEXT("Input tensor shape does not match model input size. Got %i, expected %i."), InputInputSize, ModelInputSize))
 		{
-			return ESetInputTensorShapeStatus::Fail;
+			return ESetInputTensorShapesStatus::Fail;
 		}
 
 		BatchSize = InputBatchSize;
@@ -3421,7 +3421,7 @@ namespace UE::NNE::RuntimeBasic
 			Instance->SetMaxBatchSize(BatchSize);
 		}
 
-		return ESetInputTensorShapeStatus::Ok;
+		return ESetInputTensorShapesStatus::Ok;
 	}
 
 	FModelInstanceCPU::ERunSyncStatus FModelInstanceCPU::RunSync(TConstArrayView<FTensorBindingCPU> InInputBindings, TConstArrayView<FTensorBindingCPU> InOutputBindings)
diff --git a/Engine/Plugins/Experimental/NNERuntimeBasicCpu/Source/NNERuntimeBasicCpu/Private/NNERuntimeBasicCpuModel.h b/Engine/Plugins/Experimental/NNERuntimeBasicCpu/Source/NNERuntimeBasicCpu/Private/NNERuntimeBasicCpuModel.h
--- a/Engine/Plugins/Experimental/NNERuntimeBasicCpu/Source/NNERuntimeBasicCpu/Private/NNERuntimeBasicCpuModel.h
+++ b/Engine/Plugins/Experimental/NNERuntimeBasicCpu/Source/NNERuntimeBasicCpu/Private/NNERuntimeBasicCpuModel.h
@@ -28,7 +28,7 @@ namespace UE::NNE::RuntimeBasic
 	class FModelInstanceCPU : public IModelInstanceCPU
 	{
 	public:
-		using ESetInputTensorShapeStatus = IModelInstanceCPU::ESetInputTensorShapeStatus;
+		using ESetInputTensorShapesStatus = IModelInstanceCPU::ESetInputTensorShapesStatus;
 		using ERunSyncStatus = IModelInstanceCPU::ERunSyncStatus;
 
 		FModelInstanceCPU(const TSharedPtr<FModelCPU>& InModel);
@@ -38,7 +38,7 @@ namespace UE::NNE::RuntimeBasic
 		virtual TConstArrayView<FTensorDesc> GetOutputTensorDescs() const override final { return TConstArrayView<FTensorDesc>(&OutputTensorDesc, 1); }
 		virtual TConstArrayView<FTensorShape> GetInputTensorShapes() const override final { return TConstArrayView<FTensorShape>(&InputTensorShape, 1); }
 		virtual TConstArrayView<FTensorShape> GetOutputTensorShapes() const override final { return TConstArrayView<FTensorShape>(&OutputTensorShape, 1); }
-		virtual ESetInputTensorShapeStatus SetInputTensorShapes(TConstArrayView<FTensorShape> InInputShapes) override final;
+		virtual ESetInputTensorShapesStatus SetInputTensorShapes(TConstArrayView<FTensorShape> InInputShapes) override final;
 		virtual ERunSyncStatus RunSync(TConstArrayView<FTensorBindingCPU> InInputBindings, TConstArrayView<FTensorBindingCPU> InOutputBindings) override final;
 		//~ End UE::NNE::IModelInstanceCPU Interface
 
diff --git a/Engine/Plugins/Experimental/NNERuntimeIREE/Source/NNERuntimeIREE/Private/NNERuntimeIREEModel.cpp b/Engine/Plugins/Experimental/NNERuntimeIREE/Source/NNERuntimeIREE/Private/NNERuntimeIREEModel.cpp
--- a/Engine/Plugins/Experimental/NNERuntimeIREE/Source/NNERuntimeIREE/Private/NNERuntimeIREEModel.cpp
+++ b/Engine/Plugins/Experimental/NNERuntimeIREE/Source/NNERuntimeIREE/Private/NNERuntimeIREEModel.cpp
@@ -638,7 +638,7 @@ namespace UE::NNERuntimeIREE
 				}
 
 			public:
-				using ESetInputTensorShapeStatus = UE::NNE::IModelInstanceCPU::ESetInputTensorShapeStatus;
+				using ESetInputTensorShapesStatus = UE::NNE::IModelInstanceCPU::ESetInputTensorShapesStatus;
 				using ERunSyncStatus = UE::NNE::IModelInstanceCPU::ERunSyncStatus;
 
 				~FSession()
@@ -724,12 +724,12 @@ namespace UE::NNERuntimeIREE
 					return OutputTensorShapes;
 				}
 
-				ESetInputTensorShapeStatus SetInputTensorShapes(TConstArrayView<UE::NNE::FTensorShape> InInputShapes)
+				ESetInputTensorShapesStatus SetInputTensorShapes(TConstArrayView<UE::NNE::FTensorShape> InInputShapes)
 				{
 					check(InputTensorDescs.Num() == InInputShapes.Num());
 					checkCode(for (int32 i = 0; i < InputTensorDescs.Num(); i++) { check(InInputShapes[i].IsCompatibleWith(InputTensorDescs[i].GetShape())); });
 					InputTensorShapes = InInputShapes;
-					return ESetInputTensorShapeStatus::Ok;
+					return ESetInputTensorShapesStatus::Ok;
 				}
 
 				ERunSyncStatus RunSyncCPU(TConstArrayView<UE::NNE::FTensorBindingCPU> InInputBindings, TConstArrayView<UE::NNE::FTensorBindingCPU> InOutputBindings)
@@ -898,7 +898,7 @@ namespace UE::NNERuntimeIREE
 			return Session->GetOutputTensorShapes();
 		}
 
-		FModelInstance::ESetInputTensorShapeStatus FModelInstance::SetInputTensorShapes(TConstArrayView<UE::NNE::FTensorShape> InInputShapes)
+		FModelInstance::ESetInputTensorShapesStatus FModelInstance::SetInputTensorShapes(TConstArrayView<UE::NNE::FTensorShape> InInputShapes)
 		{
 			return Session->SetInputTensorShapes(InInputShapes);
 		}
diff --git a/Engine/Plugins/Experimental/NNERuntimeIREE/Source/NNERuntimeIREE/Private/NNERuntimeIREEModel.h b/Engine/Plugins/Experimental/NNERuntimeIREE/Source/NNERuntimeIREE/Private/NNERuntimeIREEModel.h
--- a/Engine/Plugins/Experimental/NNERuntimeIREE/Source/NNERuntimeIREE/Private/NNERuntimeIREEModel.h
+++ b/Engine/Plugins/Experimental/NNERuntimeIREE/Source/NNERuntimeIREE/Private/NNERuntimeIREEModel.h
@@ -31,7 +31,7 @@ namespace UE::NNERuntimeIREE
 			static TSharedPtr<FModelInstance> Make(TSharedRef<UE::NNERuntimeIREE::CPU::Private::FDevice> InDevice, TSharedRef<UE::NNERuntimeIREE::Private::FModule> InModule);
 
 		public:
-			using ESetInputTensorShapeStatus = UE::NNE::IModelInstanceCPU::ESetInputTensorShapeStatus;
+			using ESetInputTensorShapesStatus = UE::NNE::IModelInstanceCPU::ESetInputTensorShapesStatus;
 			using ERunSyncStatus = UE::NNE::IModelInstanceCPU::ERunSyncStatus;
 
 			//~ Begin IModelInstanceCPU Interface
@@ -39,7 +39,7 @@ namespace UE::NNERuntimeIREE
 			virtual TConstArrayView<UE::NNE::FTensorDesc> GetOutputTensorDescs() const override;
 			virtual TConstArrayView<UE::NNE::FTensorShape> GetInputTensorShapes() const override;
 			virtual TConstArrayView<UE::NNE::FTensorShape> GetOutputTensorShapes() const override;
-			virtual ESetInputTensorShapeStatus SetInputTensorShapes(TConstArrayView<UE::NNE::FTensorShape> InInputShapes) override;
+			virtual ESetInputTensorShapesStatus SetInputTensorShapes(TConstArrayView<UE::NNE::FTensorShape> InInputShapes) override;
 			virtual ERunSyncStatus RunSync(TConstArrayView<UE::NNE::FTensorBindingCPU> InInputBindings, TConstArrayView<UE::NNE::FTensorBindingCPU> InOutputBindings) override;
 			//~ End IModelInstanceCPU Interface
 
diff --git a/Engine/Plugins/Experimental/NNERuntimeRDG/Source/NNERuntimeRDG/Private/NNERuntimeRDGModel.cpp b/Engine/Plugins/Experimental/NNERuntimeRDG/Source/NNERuntimeRDG/Private/NNERuntimeRDGModel.cpp
--- a/Engine/Plugins/Experimental/NNERuntimeRDG/Source/NNERuntimeRDG/Private/NNERuntimeRDGModel.cpp
+++ b/Engine/Plugins/Experimental/NNERuntimeRDG/Source/NNERuntimeRDG/Private/NNERuntimeRDGModel.cpp
@@ -100,12 +100,12 @@ bool FModelInstanceRDG::LoadModel(TConstArrayView<uint8> ModelData, FNNERuntimeF
 
 
 
-FModelInstanceRDG::ESetInputTensorShapeStatus FModelInstanceRDG::SetInputTensorShapes(TConstArrayView<NNE::FTensorShape> InInputShapes)
+FModelInstanceRDG::ESetInputTensorShapesStatus FModelInstanceRDG::SetInputTensorShapes(TConstArrayView<NNE::FTensorShape> InInputShapes)
 {
 	OutputTensorShapes.Reset(OutputTensorIndices.Num());
 
 	//Verify input shape are valid for the model and set InputTensorShapes
-	if (ESetInputTensorShapeStatus Status = FModelInstanceBase<NNE::IModelInstanceRDG>::SetInputTensorShapes(InInputShapes); Status != ESetInputTensorShapeStatus::Ok)
+	if (ESetInputTensorShapesStatus Status = FModelInstanceBase<NNE::IModelInstanceRDG>::SetInputTensorShapes(InInputShapes); Status != ESetInputTensorShapesStatus::Ok)
 	{
 		return Status;
 	}
@@ -160,7 +160,7 @@ FModelInstanceRDG::ESetInputTensorShapeStatus FModelInstanceRDG::SetInputTensorS
 	//Allow the specific runtime to run shape inference if supported
 	if (PrepareTensorShapesAndData() != 0)
 	{
-		return ESetInputTensorShapeStatus::Fail;
+		return ESetInputTensorShapesStatus::Fail;
 	}
 
 	checkCode(
@@ -182,7 +182,7 @@ FModelInstanceRDG::ESetInputTensorShapeStatus FModelInstanceRDG::SetInputTensorS
 	check(WeightTensorIndices.Num() == WeightTensorRDGs.Num());
 	check(AllTensorRDGRefs.Num() == AllSymbolicTensorDescs.Num());
 	
-	return ESetInputTensorShapeStatus::Ok;
+	return ESetInputTensorShapesStatus::Ok;
 }
 
 FRDGBufferDesc CreateRDGBufferDescForTensorRDG(const FTensorRDG& Tensor)
diff --git a/Engine/Plugins/Experimental/NNERuntimeRDG/Source/NNERuntimeRDG/Private/NNERuntimeRDGModel.h b/Engine/Plugins/Experimental/NNERuntimeRDG/Source/NNERuntimeRDG/Private/NNERuntimeRDGModel.h
--- a/Engine/Plugins/Experimental/NNERuntimeRDG/Source/NNERuntimeRDG/Private/NNERuntimeRDGModel.h
+++ b/Engine/Plugins/Experimental/NNERuntimeRDG/Source/NNERuntimeRDG/Private/NNERuntimeRDGModel.h
@@ -14,13 +14,13 @@ namespace UE::NNERuntimeRDG::Private
 	class FModelInstanceRDG : public NNE::Internal::FModelInstanceBase<NNE::IModelInstanceRDG>
 	{
 	public:
-		using ESetInputTensorShapeStatus = IModelInstanceRDG::ESetInputTensorShapeStatus;
+		using ESetInputTensorShapesStatus = IModelInstanceRDG::ESetInputTensorShapesStatus;
 		using EEnqueueRDGStatus = IModelInstanceRDG::EEnqueueRDGStatus;
 
 		FModelInstanceRDG() {};
 		virtual ~FModelInstanceRDG() = default;
 
-		virtual ESetInputTensorShapeStatus SetInputTensorShapes(TConstArrayView<NNE::FTensorShape> InInputShapes) override;
+		virtual ESetInputTensorShapesStatus SetInputTensorShapes(TConstArrayView<NNE::FTensorShape> InInputShapes) override;
 		virtual EEnqueueRDGStatus EnqueueRDG(FRDGBuilder& RDGBuilder, TConstArrayView<NNE::FTensorBindingRDG> InInputBindings, TConstArrayView<NNE::FTensorBindingRDG> InOutputBindings) override;
 		
 	protected:
diff --git a/Engine/Plugins/Experimental/NeuralRendering/Source/NeuralPostProcessing/Private/NeuralPostProcessModelInstance.cpp b/Engine/Plugins/Experimental/NeuralRendering/Source/NeuralPostProcessing/Private/NeuralPostProcessModelInstance.cpp
--- a/Engine/Plugins/Experimental/NeuralRendering/Source/NeuralPostProcessing/Private/NeuralPostProcessModelInstance.cpp
+++ b/Engine/Plugins/Experimental/NeuralRendering/Source/NeuralPostProcessing/Private/NeuralPostProcessModelInstance.cpp
@@ -261,7 +261,7 @@ void UNeuralPostProcessModelInstance::CreateDefaultNNEModel(UNNEModelData* NNEMo
 		// All dynamic dimensions are set to 1 by default.
 		ResolvedInputTensorShape = UE::NNE::FTensorShape::MakeFromSymbolic(InputShape);
 
-		if (ModelInstanceRDG->SetInputTensorShapes({ ResolvedInputTensorShape }) != UE::NNE::IModelInstanceRDG::ESetInputTensorShapeStatus::Ok)
+		if (ModelInstanceRDG->SetInputTensorShapes({ ResolvedInputTensorShape }) != UE::NNE::IModelInstanceRDG::ESetInputTensorShapesStatus::Ok)
 		{
 			ModelInstanceRDG.Reset();
 #if WITH_EDITOR
diff --git a/Engine/Plugins/NNE/NNERuntimeORT/Source/NNERuntimeORT/Private/NNERuntimeORTModel.cpp b/Engine/Plugins/NNE/NNERuntimeORT/Source/NNERuntimeORT/Private/NNERuntimeORTModel.cpp
--- a/Engine/Plugins/NNE/NNERuntimeORT/Source/NNERuntimeORT/Private/NNERuntimeORTModel.cpp
+++ b/Engine/Plugins/NNE/NNERuntimeORT/Source/NNERuntimeORT/Private/NNERuntimeORTModel.cpp
@@ -161,7 +161,7 @@ namespace UE::NNERuntimeORT::Private
 	}
 
 	template <class ModelInterface, class TensorBinding> 
-	typename ModelInterface::ESetInputTensorShapeStatus FModelInstanceORTBase<ModelInterface, TensorBinding>::SetInputTensorShapes(TConstArrayView<NNE::FTensorShape> InInputShapes)
+	typename ModelInterface::ESetInputTensorShapesStatus FModelInstanceORTBase<ModelInterface, TensorBinding>::SetInputTensorShapes(TConstArrayView<NNE::FTensorShape> InInputShapes)
 	{
 		using ModelInstanceBase = NNE::Internal::FModelInstanceBase<ModelInterface>;
 
@@ -170,7 +170,7 @@ namespace UE::NNERuntimeORT::Private
 		ModelInstanceBase::OutputTensorShapes.Reset();
 
 		// Verify input shape are valid for the model and set InputTensorShapes
-		if (typename ModelInterface::ESetInputTensorShapeStatus Status = ModelInstanceBase::SetInputTensorShapes(InInputShapes); Status != ModelInterface::ESetInputTensorShapeStatus::Ok)
+		if (typename ModelInterface::ESetInputTensorShapesStatus Status = ModelInstanceBase::SetInputTensorShapes(InInputShapes); Status != ModelInterface::ESetInputTensorShapesStatus::Ok)
 		{
 			return Status;
 		}
@@ -198,7 +198,7 @@ namespace UE::NNERuntimeORT::Private
 			ModelInstanceBase::OutputTensorShapes.Reset();
 		}
 
-		return ModelInterface::ESetInputTensorShapeStatus::Ok;
+		return ModelInterface::ESetInputTensorShapesStatus::Ok;
 	}
 
 	template <class ModelInterface, class TensorBinding>
diff --git a/Engine/Plugins/NNE/NNERuntimeORT/Source/NNERuntimeORT/Private/NNERuntimeORTModel.h b/Engine/Plugins/NNE/NNERuntimeORT/Source/NNERuntimeORT/Private/NNERuntimeORTModel.h
--- a/Engine/Plugins/NNE/NNERuntimeORT/Source/NNERuntimeORT/Private/NNERuntimeORTModel.h
+++ b/Engine/Plugins/NNE/NNERuntimeORT/Source/NNERuntimeORT/Private/NNERuntimeORTModel.h
@@ -26,7 +26,7 @@ namespace UE::NNERuntimeORT::Private
 		FModelInstanceORTBase(const FRuntimeConf& InRuntimeConf, TSharedPtr<Ort::Env> InEnvironment);
 		virtual ~FModelInstanceORTBase() = default;
 
-		virtual typename ModelInterface::ESetInputTensorShapeStatus SetInputTensorShapes(TConstArrayView<NNE::FTensorShape> InInputShapes) override;
+		virtual typename ModelInterface::ESetInputTensorShapesStatus SetInputTensorShapes(TConstArrayView<NNE::FTensorShape> InInputShapes) override;
 
 		bool Init(TConstArrayView<uint8> ModelData);
 
diff --git a/Engine/Source/Runtime/NNE/Internal/NNEModelBase.h b/Engine/Source/Runtime/NNE/Internal/NNEModelBase.h
--- a/Engine/Source/Runtime/NNE/Internal/NNEModelBase.h
+++ b/Engine/Source/Runtime/NNE/Internal/NNEModelBase.h
@@ -20,7 +20,7 @@ public:
 	virtual TConstArrayView<NNE::FTensorDesc> GetOutputTensorDescs() const override;
 	virtual TConstArrayView<NNE::FTensorShape> GetInputTensorShapes() const override;
 	virtual TConstArrayView<NNE::FTensorShape> GetOutputTensorShapes() const override;
-	virtual typename ModelInterface::ESetInputTensorShapeStatus SetInputTensorShapes(TConstArrayView<NNE::FTensorShape> InInputShapes) override;
+	virtual typename ModelInterface::ESetInputTensorShapesStatus SetInputTensorShapes(TConstArrayView<NNE::FTensorShape> InInputShapes) override;
 
 protected:
 
@@ -57,14 +57,14 @@ TConstArrayView<FTensorShape> FModelInstanceBase<T>::GetOutputTensorShapes() con
 }
 
 template <class T>
-typename T::ESetInputTensorShapeStatus FModelInstanceBase<T>::SetInputTensorShapes(TConstArrayView<FTensorShape> InInputShapes)
+typename T::ESetInputTensorShapesStatus FModelInstanceBase<T>::SetInputTensorShapes(TConstArrayView<FTensorShape> InInputShapes)
 {
 	InputTensorShapes.Reset(InInputShapes.Num());
 
 		if (InInputShapes.Num() != InputSymbolicTensors.Num())
 		{
 			UE_LOG(LogNNE, Error, TEXT("Number of input shapes does not match number of input tensors"));
-			return T::ESetInputTensorShapeStatus::Fail;
+			return T::ESetInputTensorShapesStatus::Fail;
 		}
 
 		for (int32 i = 0; i < InInputShapes.Num(); ++i)
@@ -73,7 +73,7 @@ typename T::ESetInputTensorShapeStatus FModelInstanceBase<T>::SetInputTensorShap
 			if (!InInputShapes[i].IsCompatibleWith(SymbolicDesc.GetShape()))
 			{
 				UE_LOG(LogNNE, Error, TEXT("Input shape does not match input tensor %s of index %d"), *SymbolicDesc.GetName(), i);
-				return T::ESetInputTensorShapeStatus::Fail;
+				return T::ESetInputTensorShapesStatus::Fail;
 			}
 		}
 		
@@ -81,7 +81,7 @@ typename T::ESetInputTensorShapeStatus FModelInstanceBase<T>::SetInputTensorShap
 
 	//Implementations are responsible to handle output and intermediate tensor shape inference.
 	//This base implementation only validate that all inputs are matching what the model can support.
-	return T::ESetInputTensorShapeStatus::Ok;
+	return T::ESetInputTensorShapesStatus::Ok;
 }
 
 } // namespace UE::NNE::Internal
\ No newline at end of file
diff --git a/Engine/Source/Runtime/NNE/Private/NNE.cpp b/Engine/Source/Runtime/NNE/Private/NNE.cpp
--- a/Engine/Source/Runtime/NNE/Private/NNE.cpp
+++ b/Engine/Source/Runtime/NNE/Private/NNE.cpp
@@ -81,9 +81,9 @@ namespace UE::NNE
 		TMap<FString, TWeakInterfacePtr<INNERuntime>> Runtimes;
 	};
 
-	ERegisterRuntimeResultStatus RegisterRuntime(TWeakInterfacePtr<INNERuntime> Runtime)
+	ERegisterRuntimeStatus RegisterRuntime(TWeakInterfacePtr<INNERuntime> Runtime)
 	{
-		const ERegisterRuntimeResultStatus Result = FRegistry::GetInstance().Add(Runtime);
+		const ERegisterRuntimeStatus Result = FRegistry::GetInstance().Add(Runtime);
 
 #ifdef WITH_EDITOR
 		FModuleManager::Get().LoadModule(TEXT("NNEEditor"));
@@ -106,7 +106,7 @@ namespace UE::NNE
 		return Result;
 	}
 
-	EUnregisterRuntimeResultStatus UnregisterRuntime(TWeakInterfacePtr<INNERuntime> Runtime)
+	EUnregisterRuntimeStatus UnregisterRuntime(TWeakInterfacePtr<INNERuntime> Runtime)
 	{
 		return FRegistry::GetInstance().Remove(Runtime);
 	}
diff --git a/Engine/Source/Runtime/NNE/Public/NNE.h b/Engine/Source/Runtime/NNE/Public/NNE.h
--- a/Engine/Source/Runtime/NNE/Public/NNE.h
+++ b/Engine/Source/Runtime/NNE/Public/NNE.h
@@ -13,8 +13,8 @@ NNE_API DECLARE_LOG_CATEGORY_EXTERN(LogNNE, Log, All);
 
 namespace UE::NNE
 {
-	using ERegisterRuntimeResultStatus = EResultStatus;
-	using EUnregisterRuntimeResultStatus = EResultStatus;
+	using ERegisterRuntimeStatus = EResultStatus;
+	using EUnregisterRuntimeStatus = EResultStatus;
 
 	/**
 	 * Register a runtime to make it accessible to NNE clients.
@@ -24,7 +24,7 @@ namespace UE::NNE
 	 * @param Runtime A weak interface pointer to the runtime to be registered.
 	 * @return Status indicating success or failure (e.g. if the runtime already has been registered).
 	 */
-	NNE_API ERegisterRuntimeResultStatus RegisterRuntime(TWeakInterfacePtr<INNERuntime> Runtime);
+	NNE_API ERegisterRuntimeStatus RegisterRuntime(TWeakInterfacePtr<INNERuntime> Runtime);
 
 	/**
 	 * Unregister a registered runtime.
@@ -32,7 +32,7 @@ namespace UE::NNE
 	 * @param Runtime A weak interface pointer to the runtime to be unregistered.
 	 * @return Status indicating success or failure (e.g. if the runtime has not been registered).
 	 */
-	NNE_API EUnregisterRuntimeResultStatus UnregisterRuntime(TWeakInterfacePtr<INNERuntime> Runtime);
+	NNE_API EUnregisterRuntimeStatus UnregisterRuntime(TWeakInterfacePtr<INNERuntime> Runtime);
 	
 	/**
 	 * List and return all registered runtime names.
diff --git a/Engine/Source/Runtime/NNE/Public/NNERuntimeCPU.h b/Engine/Source/Runtime/NNE/Public/NNERuntimeCPU.h
--- a/Engine/Source/Runtime/NNE/Public/NNERuntimeCPU.h
+++ b/Engine/Source/Runtime/NNE/Public/NNERuntimeCPU.h
@@ -35,7 +35,7 @@ class NNE_API IModelInstanceCPU
 {
 public:
 
-	using ESetInputTensorShapeStatus = EResultStatus;
+	using ESetInputTensorShapesStatus = EResultStatus;
 	using ERunSyncStatus = EResultStatus;
 
 	virtual ~IModelInstanceCPU() = default;
@@ -83,7 +83,7 @@ public:
 	 * @param InInputShapes The input shapes to prepare the model with.
 	 * @return Status indicating success or failure.
 	 */
-	virtual ESetInputTensorShapeStatus SetInputTensorShapes(TConstArrayView<FTensorShape> InInputShapes) = 0;
+	virtual ESetInputTensorShapesStatus SetInputTensorShapes(TConstArrayView<FTensorShape> InInputShapes) = 0;
 
 	/**
 	 * Evaluate the model synchronously.
diff --git a/Engine/Source/Runtime/NNE/Public/NNERuntimeGPU.h b/Engine/Source/Runtime/NNE/Public/NNERuntimeGPU.h
--- a/Engine/Source/Runtime/NNE/Public/NNERuntimeGPU.h
+++ b/Engine/Source/Runtime/NNE/Public/NNERuntimeGPU.h
@@ -35,7 +35,7 @@ class NNE_API IModelInstanceGPU
 {
 public:
 
-	using ESetInputTensorShapeStatus = EResultStatus;
+	using ESetInputTensorShapesStatus = EResultStatus;
 	using ERunSyncStatus = EResultStatus;
 
 	virtual ~IModelInstanceGPU() = default;
@@ -83,7 +83,7 @@ public:
 	 * @param InInputShapes The input shapes to prepare the model with.
 	 * @return Status indicating success or failure.
 	 */
-	virtual ESetInputTensorShapeStatus SetInputTensorShapes(TConstArrayView<FTensorShape> InInputShapes) = 0;
+	virtual ESetInputTensorShapesStatus SetInputTensorShapes(TConstArrayView<FTensorShape> InInputShapes) = 0;
 
 	/**
 	 * Evaluate the model synchronously.
diff --git a/Engine/Source/Runtime/NNE/Public/NNERuntimeRDG.h b/Engine/Source/Runtime/NNE/Public/NNERuntimeRDG.h
--- a/Engine/Source/Runtime/NNE/Public/NNERuntimeRDG.h
+++ b/Engine/Source/Runtime/NNE/Public/NNERuntimeRDG.h
@@ -35,7 +35,7 @@ class NNE_API IModelInstanceRDG
 {
 public:
 
-	using ESetInputTensorShapeStatus = EResultStatus;
+	using ESetInputTensorShapesStatus = EResultStatus;
 	using EEnqueueRDGStatus = EResultStatus;
 
 	virtual ~IModelInstanceRDG() = default;
@@ -84,7 +84,7 @@ public:
 	 * @param InInputShapes The input shapes to prepare the model with.
 	 * @return Status indicating success or failure.
 	 */
-	virtual ESetInputTensorShapeStatus SetInputTensorShapes(TConstArrayView<FTensorShape> InInputShapes) = 0;
+	virtual ESetInputTensorShapesStatus SetInputTensorShapes(TConstArrayView<FTensorShape> InInputShapes) = 0;
 
 	/**
 	 * Enqueue the model graph to a FRDGBuilder.
