[NNE] Fix: Allow model creation on any thread for the editor case.

Problem:
* in editor model creation is accessing cvars, accesses currently locked to game thread

Solution:
* access those cvars allowing any threads.

#tests win64 editor: smoke tests, win64 standalone: all tests (dml meta command disabled)

#rnx

diff --git a/Engine/Source/Runtime/NNE/Private/NNEModelData.cpp b/Engine/Source/Runtime/NNE/Private/NNEModelData.cpp
--- a/Engine/Source/Runtime/NNE/Private/NNEModelData.cpp
+++ b/Engine/Source/Runtime/NNE/Private/NNEModelData.cpp
@@ -268,8 +268,8 @@ void UNNEModelData::Serialize(FArchive& Ar)
 						ModelDataIdentifier = NNERuntime->GetModelDataIdentifier(FileType, FileData, AdditionalFileDataView, FileId, Ar.GetArchiveState().CookingTarget());
 						if (ModelDataIdentifier.Len() > 0)
 						{
-							if (UE::NNE::ConsoleVariables::CVarNNEEditorUseDDC.GetValueOnGameThread() &&
-								UE::NNE::ConsoleVariables::CVarNNEEditorUseDDCForCookingDeprecated.GetValueOnGameThread())
+							if (UE::NNE::ConsoleVariables::CVarNNEEditorUseDDC.GetValueOnAnyThread() &&
+								UE::NNE::ConsoleVariables::CVarNNEEditorUseDDCForCookingDeprecated.GetValueOnAnyThread())
 							{
 								SharedModelData = UE::NNE::ModelData::GetFromDDC(FileId, RuntimeName, ModelDataIdentifier);
 							}
@@ -566,7 +566,7 @@ TSharedPtr<UE::NNE::FSharedModelData> UNNEModelData::GetModelData(const FString&
 		return TSharedPtr<UE::NNE::FSharedModelData>();
 	}
 
-	if (UE::NNE::ConsoleVariables::CVarNNEEditorUseDDC.GetValueOnGameThread())
+	if (UE::NNE::ConsoleVariables::CVarNNEEditorUseDDC.GetValueOnAnyThread())
 	{
 		// Check if we have a DDC cache hit
 		TSharedPtr<UE::NNE::FSharedModelData> RemoteData = UE::NNE::ModelData::GetFromDDC(FileId, RuntimeName, ModelDataIdentifier);
