[NNE] Remove DDC from UNNEModelData.

#tests Editor Win64 smoke tests

diff --git a/Engine/Source/Editor/NNEEditor/Private/NNEEditorSettings.cpp b/Engine/Source/Editor/NNEEditor/Private/NNEEditorSettings.cpp
deleted file mode 100644
diff --git a/Engine/Source/Editor/NNEEditor/Private/NNEEditorSettings.h b/Engine/Source/Editor/NNEEditor/Private/NNEEditorSettings.h
deleted file mode 100644
diff --git a/Engine/Source/Runtime/NNE/Private/NNEModelData.cpp b/Engine/Source/Runtime/NNE/Private/NNEModelData.cpp
--- a/Engine/Source/Runtime/NNE/Private/NNEModelData.cpp
+++ b/Engine/Source/Runtime/NNE/Private/NNEModelData.cpp
@@ -11,15 +11,6 @@
 #include "UObject/AssetRegistryTagsContext.h"
 #include "UObject/WeakInterfacePtr.h"
 
-#if WITH_EDITOR
-#include "Containers/StringFwd.h"
-#include "DerivedDataCache.h"
-#include "DerivedDataCacheKey.h"
-#include "DerivedDataRequestOwner.h"
-#include "HAL/IConsoleManager.h"
-#include "Memory/CompositeBuffer.h"
-#endif // WITH_EDITOR
-
 namespace UE::NNE::ModelData
 {
 	enum Version : uint32
@@ -36,8 +27,6 @@ namespace UE::NNE::ModelData
 	const FGuid GUID(0x9513202e, 0xeba1b279, 0xf17fe5ba, 0xab90c3f2);
 	FCustomVersionRegistration NNEModelDataVersion(GUID, Version::Latest, TEXT("NNEModelDataVersion"));// Always save with the latest version
 
-	const uint32 DDCAssetVersion = 1; // Increase this value to force rebuilding cache entries
-
 	FString GetRuntimesAsString(TArrayView<const FString> Runtimes)
 	{
 		if (Runtimes.Num() == 0)
@@ -60,77 +49,6 @@ namespace UE::NNE::ModelData
 		return RuntimesAsOneString;
 	}
 
-#if WITH_EDITOR
-
-	inline FString GetDDCRequestId(const FString& FileId, const FString& RuntimeName, const FString& ModelDataIdentifier)
-	{
-		// RuntimeName and FileId are embedded to the id to ensure no potential collision between the runtime/assets
-		return RuntimeName + "-" + FileId + "-DDCv" + FString::FromInt(DDCAssetVersion) + "-" + ModelDataIdentifier;
-	}
-
-	inline UE::DerivedData::FCacheKey CreateCacheKey(const FString& RequestId)
-	{
-		return { UE::DerivedData::FCacheBucket(TEXT("NNEModelData")), FIoHash::HashBuffer(MakeMemoryView(FTCHARToUTF8(RequestId))) };
-	}
-
-	inline void PutIntoDDC(const FGuid& FileId, const FString& RuntimeName, const FString& ModelDataIdentifier, TSharedPtr<UE::NNE::FSharedModelData> Data)
-	{
-		check(Data.IsValid());
-		check(Data->GetView().Num() > 0);
-
-		FString FileIdStr = FileId.ToString(EGuidFormats::Digits);
-		FString RequestId = GetDDCRequestId(FileIdStr, RuntimeName, ModelDataIdentifier);
-
-		TArray<UE::DerivedData::FCachePutValueRequest> Requests;
-		Requests.SetNum(1);
-
-		Requests[0].Name = FString("Put-") + RequestId;
-		Requests[0].Key = CreateCacheKey(RequestId);
-		Requests[0].Value = UE::DerivedData::FValue::Compress(FCompositeBuffer(MakeSharedBufferFromArray(TArray<uint32>({ Data->GetMemoryAlignment() })), FSharedBuffer::MakeView(Data->GetView().GetData(), Data->GetView().Num())));
-
-		UE::DerivedData::FRequestOwner BlockingPutOwner(UE::DerivedData::EPriority::Blocking);
-		UE::DerivedData::GetCache().PutValue(Requests, BlockingPutOwner);
-		BlockingPutOwner.Wait();
-	}
-
-	inline TSharedPtr<UE::NNE::FSharedModelData> GetFromDDC(const FGuid& FileId, const FString& RuntimeName, const FString& ModelDataIdentifier)
-	{
-		FString FileIdStr = FileId.ToString(EGuidFormats::Digits);
-		FString RequestId = GetDDCRequestId(FileIdStr, RuntimeName, ModelDataIdentifier);
-
-		TArray<UE::DerivedData::FCacheGetValueRequest> Requests;
-		Requests.SetNum(1);
-
-		Requests[0].Name = FString("Get-") + RequestId;
-		Requests[0].Key = CreateCacheKey(RequestId);
-
-		TSharedPtr<UE::NNE::FSharedModelData> Result;
-		UE::DerivedData::FRequestOwner BlockingGetOwner(UE::DerivedData::EPriority::Blocking);
-		UE::DerivedData::GetCache().GetValue(Requests, BlockingGetOwner, [&Result](UE::DerivedData::FCacheGetValueResponse&& Response)
-		{
-			if (Response.Value.HasData() && Response.Value.GetRawSize() > sizeof(uint32))
-			{
-				FCompressedBufferReader Reader(Response.Value.GetData());
-				uint32 MemoryAlignment = ((uint32*)Reader.Decompress(0, sizeof(uint32)).GetData())[0];
-				uint64 DataSize = Response.Value.GetRawSize() - sizeof(uint32);
-
-				void* Data = FMemory::Malloc(DataSize, MemoryAlignment);
-				if (Reader.TryDecompressTo(FMutableMemoryView(Data, DataSize), sizeof(uint32)))
-				{
-					Result = MakeShared<UE::NNE::FSharedModelData>(FSharedBuffer::TakeOwnership(Data, DataSize, FMemory::Free), MemoryAlignment);
-				}
-				else
-				{
-					FMemory::Free(Data);
-				}
-			}
-		});
-		BlockingGetOwner.Wait();
-		return Result;
-	}
-
-#endif // WITH_EDITOR
-
 	inline TSharedPtr<UE::NNE::FSharedModelData> CreateModelData(const FString& RuntimeName, const FString& FileType, const TArray<uint8>& FileData, const TMap<FString, TConstArrayView<uint8>>& AdditionalFileData, const FGuid& FileId, const ITargetPlatform* TargetPlatform)
 	{
 		TWeakInterfacePtr<INNERuntime> NNERuntime = UE::NNE::GetRuntime<INNERuntime>(RuntimeName);
@@ -186,24 +104,6 @@ namespace UE::NNE
 	{
 		return MemoryAlignment;
 	}
-
-#if WITH_EDITOR
-	namespace ConsoleVariables
-	{
-		static TAutoConsoleVariable<bool> CVarNNEEditorUseDDC(
-			TEXT("nne.Editor.UseDDC"),
-			true,
-			TEXT("Indicates whether NNE should use the DDC cache to speed up SharedModelData creation in the editor, see also nne.Editor.UseDDCForCooking."),
-			ECVF_Default);
-
-		static TAutoConsoleVariable<bool> CVarNNEEditorUseDDCForCookingDeprecated(
-			TEXT("nne.Editor.UseDDCForCooking_Deprecated"),
-			false,
-			TEXT("DEPRECATED. Indicates whether NNE should use the DDC cache to speed up SharedModelData creation while cooking, if nne.Editor.UseDDC is false cooking won't use the DDC even if nne.Editor.UseDDCForCooking is true."),
-			ECVF_Default);
-	} // ConsoleVariables
-#endif //WITH_EDITOR
-
 } // UE::NNE
 
 void UNNEModelData::GetAssetRegistryTags(FAssetRegistryTagsContext Context) const
@@ -259,41 +159,7 @@ void UNNEModelData::Serialize(FArchive& Ar)
 
 				for (const FString& RuntimeName : CookRuntimeNames)
 				{
-					TSharedPtr<UE::NNE::FSharedModelData> SharedModelData;
-#if WITH_EDITOR
-					TWeakInterfacePtr<INNERuntime> NNERuntime = UE::NNE::GetRuntime<INNERuntime>(RuntimeName);
-					FString ModelDataIdentifier;
-					if (NNERuntime.IsValid())
-					{
-						ModelDataIdentifier = NNERuntime->GetModelDataIdentifier(FileType, FileData, AdditionalFileDataView, FileId, Ar.GetArchiveState().CookingTarget());
-						if (ModelDataIdentifier.Len() > 0)
-						{
-							if (UE::NNE::ConsoleVariables::CVarNNEEditorUseDDC.GetValueOnAnyThread() &&
-								UE::NNE::ConsoleVariables::CVarNNEEditorUseDDCForCookingDeprecated.GetValueOnAnyThread())
-							{
-								SharedModelData = UE::NNE::ModelData::GetFromDDC(FileId, RuntimeName, ModelDataIdentifier);
-							}
-						}
-						else
-						{
-							UE_LOG(LogNNE, Warning, TEXT("UNNEModelData: Runtime '%s' returned an empty string as a ModelDataIdentifier while cooking. GetModelDataIdentifier should always return a valid identifier."), *RuntimeName);
-						}
-					}
-					else
-					{
-						UE_LOG(LogNNE, Warning, TEXT("UNNEModelData: Runtime '%s' is among the cooked runtimes but instance is invalid."), *RuntimeName);
-					}
-#endif //WITH_EDITOR
-					if (!SharedModelData.IsValid() || SharedModelData->GetView().Num() == 0)
-					{
-						SharedModelData = UE::NNE::ModelData::CreateModelData(RuntimeName, FileType, FileData, AdditionalFileDataView, FileId, Ar.GetArchiveState().CookingTarget());
-#if WITH_EDITOR
-						if (SharedModelData.IsValid() && (SharedModelData->GetView().Num() > 0) && NNERuntime.IsValid() && (ModelDataIdentifier.Len() > 0))
-						{
-							UE::NNE::ModelData::PutIntoDDC(FileId, RuntimeName, ModelDataIdentifier, SharedModelData);
-						}
-#endif //WITH_EDITOR
-					}
+					TSharedPtr<UE::NNE::FSharedModelData> SharedModelData = UE::NNE::ModelData::CreateModelData(RuntimeName, FileType, FileData, AdditionalFileDataView, FileId, Ar.GetArchiveState().CookingTarget());
 					if (SharedModelData.IsValid() && SharedModelData->GetView().Num() > 0)
 					{
 						ModelData.Add(RuntimeName, SharedModelData);
@@ -551,33 +417,6 @@ TSharedPtr<UE::NNE::FSharedModelData> UNNEModelData::GetModelData(const FString&
 		AdditionalFileDataView.Emplace(Key, AdditionalFileData[Key]);
 	}
 
-#if WITH_EDITOR
-	TWeakInterfacePtr<INNERuntime> NNERuntime = UE::NNE::GetRuntime<INNERuntime>(RuntimeName);
-	if (!NNERuntime.IsValid())
-	{
-		UE_LOG(LogNNE, Error, TEXT("UNNEModelData: Runtime '%s' is among the target runtimes but instance is invalid."), *RuntimeName);
-		return TSharedPtr<UE::NNE::FSharedModelData>();
-	}
-
-	FString ModelDataIdentifier = NNERuntime->GetModelDataIdentifier(FileType, FileData, AdditionalFileDataView, FileId, nullptr);
-	if (ModelDataIdentifier.Len() == 0)
-	{
-		UE_LOG(LogNNE, Error, TEXT("UNNEModelData: Runtime '%s' returned an empty string as a ModelDataIdentifier. GetModelDataIdentifier should always return a valid identifier."), *RuntimeName);
-		return TSharedPtr<UE::NNE::FSharedModelData>();
-	}
-
-	if (UE::NNE::ConsoleVariables::CVarNNEEditorUseDDC.GetValueOnAnyThread())
-	{
-		// Check if we have a DDC cache hit
-		TSharedPtr<UE::NNE::FSharedModelData> RemoteData = UE::NNE::ModelData::GetFromDDC(FileId, RuntimeName, ModelDataIdentifier);
-		if (RemoteData.IsValid() && RemoteData->GetView().Num() > 0)
-		{
-			ModelData.Add(RuntimeName, RemoteData);
-			return RemoteData;
-		}
-	}
-#endif //WITH_EDITOR
-
 	// Try to create the model
 	TSharedPtr<UE::NNE::FSharedModelData> CreatedData = UE::NNE::ModelData::CreateModelData(RuntimeName, FileType, FileData, AdditionalFileDataView, FileId, nullptr);
 	if (!CreatedData.IsValid() || CreatedData->GetView().Num() < 1)
@@ -588,11 +427,6 @@ TSharedPtr<UE::NNE::FSharedModelData> UNNEModelData::GetModelData(const FString&
 	// Cache the model
 	ModelData.Add(RuntimeName, CreatedData);
 
-#if WITH_EDITOR
-	// And put it into DDC
-	UE::NNE::ModelData::PutIntoDDC(FileId, RuntimeName, ModelDataIdentifier, CreatedData);
-#endif //WITH_EDITOR
-
 	return CreatedData;
 }
 
