[NNE] Break dependency to Engine.

#tests Editor, Tools & Monolithics preflight and local compilation tests.

diff --git a/Engine/Source/Runtime/NNE/NNE.Build.cs b/Engine/Source/Runtime/NNE/NNE.Build.cs
--- a/Engine/Source/Runtime/NNE/NNE.Build.cs
+++ b/Engine/Source/Runtime/NNE/NNE.Build.cs
@@ -10,8 +10,7 @@ namespace UnrealBuildTool.Rules
 				new string[]
 				{
 					"Core",
-					"CoreUObject",
-					"Engine",
+					"CoreUObject"
 				}
 			);
 
diff --git a/Engine/Source/Runtime/NNE/Private/NNE.cpp b/Engine/Source/Runtime/NNE/Private/NNE.cpp
--- a/Engine/Source/Runtime/NNE/Private/NNE.cpp
+++ b/Engine/Source/Runtime/NNE/Private/NNE.cpp
@@ -2,10 +2,9 @@
 
 #include "NNE.h"
 
-#include "EngineAnalytics.h"
-#include "Kismet/GameplayStatics.h"
 #include "Misc/CoreDelegates.h"
 #include "Misc/SecureHash.h"
+#include "Modules/ModuleManager.h"
 
 DEFINE_LOG_CATEGORY(LogNNE);
 
@@ -83,27 +82,11 @@ namespace UE::NNE
 
 	ERegisterRuntimeStatus RegisterRuntime(TWeakInterfacePtr<INNERuntime> Runtime)
 	{
-		const ERegisterRuntimeStatus Result = FRegistry::GetInstance().Add(Runtime);
-
 #ifdef WITH_EDITOR
 		FModuleManager::Get().LoadModule(TEXT("NNEEditor"));
 #endif
 
-		const FString RuntimeName = Runtime->GetRuntimeName();
-
-		FCoreDelegates::OnAllModuleLoadingPhasesComplete.AddLambda([RuntimeName]()
-		{
-			if (FEngineAnalytics::IsAvailable())
-			{
-				TArray<FAnalyticsEventAttribute> Attributes = MakeAnalyticsEventAttributeArray(
-					TEXT("PlatformName"), UGameplayStatics::GetPlatformName(),
-					TEXT("HashedRuntimeName"), FMD5::HashAnsiString(*RuntimeName)
-				);
-				FEngineAnalytics::GetProvider().RecordEvent(TEXT("NeuralNetworkEngine.RegisterRuntime"), Attributes);
-			}
-		});
-
-		return Result;
+		return FRegistry::GetInstance().Add(Runtime);
 	}
 
 	EUnregisterRuntimeStatus UnregisterRuntime(TWeakInterfacePtr<INNERuntime> Runtime)
diff --git a/Engine/Source/Runtime/NNE/Private/NNEModelData.cpp b/Engine/Source/Runtime/NNE/Private/NNEModelData.cpp
--- a/Engine/Source/Runtime/NNE/Private/NNEModelData.cpp
+++ b/Engine/Source/Runtime/NNE/Private/NNEModelData.cpp
@@ -2,7 +2,6 @@
 
 #include "NNEModelData.h"
 
-#include "EditorFramework/AssetImportData.h"
 #include "NNE.h"
 #include "NNEAttributeMap.h"
 #include "NNEModelOptimizerInterface.h"
diff --git a/Engine/Source/Runtime/NNE/Public/NNERuntimeRDG.h b/Engine/Source/Runtime/NNE/Public/NNERuntimeRDG.h
--- a/Engine/Source/Runtime/NNE/Public/NNERuntimeRDG.h
+++ b/Engine/Source/Runtime/NNE/Public/NNERuntimeRDG.h
@@ -5,11 +5,15 @@
 #include "Containers/ContainersFwd.h"
 #include "NNEStatus.h"
 #include "NNETypes.h"
-#include "RenderGraphFwd.h"
 #include "UObject/Interface.h"
 
 #include "NNERuntimeRDG.generated.h"
 
+class FRDGBuffer;
+using FRDGBufferRef = FRDGBuffer*;
+
+class FRDGBuilder;
+
 class UNNEModelData;
 
 namespace UE::NNE
