[NNE] Add a NPU interface + factorize GPU, CPU and NPU model instance into a common ancestor

Problem:
* NNE does not currently have an interface expressing that a model will run on NPU.
* GPU, CPU and NPU are all using host memory and a syncronous call but the is no way to factorize client code atm

Solution:
* Add the NPU interface.
* Add a common ancestor for GPU, CPU and NPU model instance so client code can be factorized once model instance is created.

Tests:
* editor win64 + packaged win64 : all nne tests + style transfer
* console : nne smoke test + style transfer

#rnx

diff --git a/Engine/Source/Runtime/NNE/Public/NNERuntimeCPU.h b/Engine/Source/Runtime/NNE/Public/NNERuntimeCPU.h
--- a/Engine/Source/Runtime/NNE/Public/NNERuntimeCPU.h
+++ b/Engine/Source/Runtime/NNE/Public/NNERuntimeCPU.h
@@ -2,8 +2,7 @@
 
 #pragma once
 
-#include "Containers/ContainersFwd.h"
-#include "NNETypes.h"
+#include "NNERuntimeRunSync.h"
 #include "NNEStatus.h"
 #include "UObject/Interface.h"
 
@@ -14,90 +13,14 @@ class UNNEModelData;
 namespace UE::NNE
 {
 
-/**
- * The tensor binding for passing input and output to IModelCPU.
- *
- * Memory is owned by the caller. The caller must make sure the buffer is large enough and at least as large as SizeInBytes.
- */
-struct NNE_API FTensorBindingCPU
-{
-	void*	Data;
-	uint64	SizeInBytes;
-};
-
 /**
  * The interface of a model instance that can run on CPU.
  *
  * Use UE::NNE::IModelCPU::CreateModelInstance() to get a model instance.
  * Use UE::NNE::GetRuntime<INNERuntimeCPU>(RuntimeName) to get a runtime capable of creating CPU models.
  */
-class NNE_API IModelInstanceCPU
+class NNE_API IModelInstanceCPU : public IModelInstanceRunSync
 {
-public:
-
-	using ESetInputTensorShapesStatus = EResultStatus;
-	using ERunSyncStatus = EResultStatus;
-
-	virtual ~IModelInstanceCPU() = default;
-
-	/**
-	 * Get the input tensor descriptions as defined by the model, potentially with variable dimensions.
-	 *
-	 * @return An array containing a tensor descriptor for each input tensor of the model.
-	 */
-	virtual TConstArrayView<FTensorDesc> GetInputTensorDescs() const = 0;
-	
-	/**
-	 * Get the output tensor descriptions as defined by the model, potentially with variable dimensions.
-	 *
-	 * @return An array containing a tensor descriptor for each output tensor of the model.
-	 */
-	virtual TConstArrayView<FTensorDesc> GetOutputTensorDescs() const = 0;
-
-	/**
-	 * Get the input shapes.
-	 *
-	 * SetInputTensorShapes must be called prior of running a model.
-	 *
-	 * @return An array of input shapes or an empty array if SetInputTensorShapes has not been called.
-	 */
-	virtual TConstArrayView<FTensorShape> GetInputTensorShapes() const = 0;
-
-	/**
-	 * Getters for outputs shapes if they were already resolved.
-	 *
-	 * Output shapes might be resolved after a call to SetInputTensorShapes if the model and runtime supports it.
-	 * Otherwise they will be resolved while running the model
-	 *
-	 * @return An array of output shapes or an empty array if not resolved yet.
-	 */
-	virtual TConstArrayView<FTensorShape> GetOutputTensorShapes() const = 0;
-	
-	/**
-	 * Prepare the model to be run with the given input shape.
-	 *
-	 * The call is mandatory before a model can be run.
-	 * The function will run shape inference and resolve, if possible, the output shapes which can then be accessed by calling GetOutputTensorShapes().
-	 * This is a potentially expensive call and should be called lazily if possible.
-	 *
-	 * @param InInputShapes The input shapes to prepare the model with.
-	 * @return Status indicating success or failure.
-	 */
-	virtual ESetInputTensorShapesStatus SetInputTensorShapes(TConstArrayView<FTensorShape> InInputShapes) = 0;
-
-	/**
-	 * Evaluate the model synchronously.
-	 *
-	 * SetInputTensorShapes must be called prior to this call.
-	 * This function will block the calling thread until the inference is complete.
-	 * The caller owns the memory inside the bindings and must make sure that they are big enough.
-	 * Clients can call this function from an async task but must make sure the memory remains valid throughout the evaluation.
-	 *
-	 * @param InInputTensors An array containing tensor bindings for each input tensor with caller owned memory containing the input data.
-	 * @param InOutputTensors An array containing tensor bindings for each output tensor with caller owned memory big enough to contain the results on success.
-	 * @return Status indicating success or failure.
-	 */
-	virtual ERunSyncStatus RunSync(TConstArrayView<FTensorBindingCPU> InInputTensors, TConstArrayView<FTensorBindingCPU> InOutputTensors) = 0;
 };
 
 /**
diff --git a/Engine/Source/Runtime/NNE/Public/NNERuntimeGPU.h b/Engine/Source/Runtime/NNE/Public/NNERuntimeGPU.h
--- a/Engine/Source/Runtime/NNE/Public/NNERuntimeGPU.h
+++ b/Engine/Source/Runtime/NNE/Public/NNERuntimeGPU.h
@@ -2,9 +2,8 @@
 
 #pragma once
 
-#include "Containers/ContainersFwd.h"
+#include "NNERuntimeRunSync.h"
 #include "NNEStatus.h"
-#include "NNETypes.h"
 #include "UObject/Interface.h"
 
 #include "NNERuntimeGPU.generated.h"
@@ -13,17 +12,7 @@ class UNNEModelData;
 
 namespace UE::NNE
 {
-
-/**
- * The tensor binding for passing input and output to IModelGPU.
- *
- * Memory is owned by the caller. The caller must make sure the buffer is large enough and at least as large as SizeInBytes.
- */
-struct NNE_API FTensorBindingGPU
-{
-	void*	Data;
-	uint64	SizeInBytes;
-};
+typedef FTensorBindingCPU FTensorBindingGPU;
 
 /**
  * The interface of a model instance that can run on GPU.
@@ -31,73 +20,8 @@ struct NNE_API FTensorBindingGPU
  * Use UE::NNE::IModelGPU::CreateModelInstance() to get a model instance.
  * Use UE::NNE::GetRuntime<INNERuntimeGPU>(RuntimeName) to get a runtime capable of creating GPU models.
  */
-class NNE_API IModelInstanceGPU
+class NNE_API IModelInstanceGPU : public IModelInstanceRunSync
 {
-public:
-
-	using ESetInputTensorShapesStatus = EResultStatus;
-	using ERunSyncStatus = EResultStatus;
-
-	virtual ~IModelInstanceGPU() = default;
-
-	/**
-	 * Get the input tensor descriptions as defined by the model, potentially with variable dimensions.
-	 *
-	 * @return An array containing a tensor descriptor for each input tensor of the model.
-	 */
-	virtual TConstArrayView<FTensorDesc> GetInputTensorDescs() const = 0;
-	
-	/**
-	 * Get the output tensor descriptions as defined by the model, potentially with variable dimensions.
-	 *
-	 * @return An array containing a tensor descriptor for each output tensor of the model.
-	 */
-	virtual TConstArrayView<FTensorDesc> GetOutputTensorDescs() const = 0;
-
-	/**
-	 * Get the input shapes.
-	 *
-	 * SetInputTensorShapes must be called prior of running a model.
-	 *
-	 * @return An array of input shapes or an empty array if SetInputTensorShapes has not been called.
-	 */
-	virtual TConstArrayView<FTensorShape> GetInputTensorShapes() const = 0;
-
-	/**
-	 * Getters for outputs shapes if they were already resolved.
-	 *
-	 * Output shapes might be resolved after a call to SetInputTensorShapes if the model and runtime supports it.
-	 * Otherwise they will be resolved while running the model
-	 *
-	 * @return An array of output shapes or an empty array if not resolved yet.
-	 */
-	virtual TConstArrayView<FTensorShape> GetOutputTensorShapes() const = 0;
-	
-	/**
-	 * Prepare the model to be run with the given input shape.
-	 *
-	 * The call is mandatory before a model can be run.
-	 * The function will run shape inference and resolve, if possible, the output shapes which can then be accessed by calling GetOutputTensorShapes().
-	 * This is a potentially expensive call and should be called lazily if possible.
-	 *
-	 * @param InInputShapes The input shapes to prepare the model with.
-	 * @return Status indicating success or failure.
-	 */
-	virtual ESetInputTensorShapesStatus SetInputTensorShapes(TConstArrayView<FTensorShape> InInputShapes) = 0;
-
-	/**
-	 * Evaluate the model synchronously.
-	 *
-	 * SetInputTensorShapes must be called prior to this call.
-	 * This function will block the calling thread until the inference is complete.
-	 * The caller owns the memory inside the bindings and must make sure that they are big enough.
-	 * Clients can call this function from an async task but must make sure the memory remains valid throughout the evaluation.
-	 *
-	 * @param InInputTensors An array containing tensor bindings for each input tensor with caller owned memory containing the input data.
-	 * @param InOutputTensors An array containing tensor bindings for each output tensor with caller owned memory big enough to contain the results on success.
-	 * @return Status indicating success or failure.
-	 */
-	virtual ERunSyncStatus RunSync(TConstArrayView<FTensorBindingGPU> InInputTensors, TConstArrayView<FTensorBindingGPU> InOutputTensors) = 0;
 };
 
 /**
diff --git a/Engine/Source/Runtime/NNE/Public/NNERuntimeNPU.h b/Engine/Source/Runtime/NNE/Public/NNERuntimeNPU.h
new file mode 100644
--- /dev/null
+++ b/Engine/Source/Runtime/NNE/Public/NNERuntimeNPU.h
@@ -0,0 +1,88 @@
+// Copyright Epic Games, Inc. All Rights Reserved.
+
+#pragma once
+
+#include "NNERuntimeRunSync.h"
+#include "NNEStatus.h"
+#include "UObject/Interface.h"
+
+#include "NNERuntimeNPU.generated.h"
+
+class UNNEModelData;
+
+namespace UE::NNE
+{
+
+/**
+ * The interface of a model instance that can run on NPU.
+ *
+ * Use UE::NNE::IModelNPU::CreateModelInstance() to get a model instance.
+ * Use UE::NNE::GetRuntime<INNERuntimeNPU>(RuntimeName) to get a runtime capable of creating NPU models.
+ */
+class NNE_API IModelInstanceNPU : public IModelInstanceRunSync
+{
+};
+
+/**
+ * The interface of a model capable of creating model instance that can run on NPU.
+ *
+ * Use UE::NNE::GetRuntime<INNERuntimeNPU>(RuntimeName) to get a runtime capable of creating NPU models.
+ */
+class NNE_API IModelNPU
+{
+public:
+
+	virtual ~IModelNPU() = default;
+
+	/**
+	 * Create a model instance for inference
+	 *
+	 * The runtime have the opportunity to share the model weights among multiple IModelInstanceNPU created from an IModelNPU instance, however this is not mandatory.
+	 * The caller can decide to convert the result into a shared pointer if required (e.g. if the model needs to be shared with an async task for evaluation).
+	 *
+	 * @return A caller owned model representing the neural network instance created.
+	 */
+	virtual TSharedPtr<UE::NNE::IModelInstanceNPU> CreateModelInstanceNPU() = 0;
+};
+
+} // UE::NNE
+
+UINTERFACE()
+class NNE_API UNNERuntimeNPU : public UInterface
+{
+	GENERATED_BODY()
+};
+
+/**
+ * The interface of a neural network runtime capable of creating NPU models.
+ *
+ * Call UE::NNE::GetRuntime<INNERuntimeNPU>(RuntimeName) to get a runtime implementing this interface.
+ */
+class NNE_API INNERuntimeNPU
+{
+	GENERATED_BODY()
+	
+public:
+
+	using ECanCreateModelNPUStatus = UE::NNE::EResultStatus;
+
+	/**
+	 * Check if the runtime is able to create a model given some ModelData.
+	 *
+	 * @param ModelData The model data for which to create a model.
+	 * @return True if the runtime is able to create the model, false otherwise.
+	 */
+	virtual ECanCreateModelNPUStatus CanCreateModelNPU(const TObjectPtr<UNNEModelData> ModelData) const = 0;
+	
+	/**
+	 * Create a model given some ModelData.
+	 *
+	 * The caller must make sure ModelData remains valid throughout the call.
+	 * ModelData is not required anymore after the model has been created.
+	 * The caller can decide to convert the result into a shared pointer if required (e.g. if the model needs to be shared with an async task for evaluation).
+	 *
+	 * @param ModelData The model data for which to create a model.
+	 * @return A caller owned model representing the neural network created from ModelData.
+	 */
+	virtual TSharedPtr<UE::NNE::IModelNPU> CreateModelNPU(const TObjectPtr<UNNEModelData> ModelData) = 0;
+};
\ No newline at end of file
diff --git a/Engine/Source/Runtime/NNE/Public/NNERuntimeRunSync.h b/Engine/Source/Runtime/NNE/Public/NNERuntimeRunSync.h
new file mode 100644
--- /dev/null
+++ b/Engine/Source/Runtime/NNE/Public/NNERuntimeRunSync.h
@@ -0,0 +1,99 @@
+// Copyright Epic Games, Inc. All Rights Reserved.
+
+#pragma once
+
+#include "Containers/ContainersFwd.h"
+#include "NNEStatus.h"
+#include "NNETypes.h"
+
+
+namespace UE::NNE
+{
+
+/**
+ * The tensor binding for passing input and output to from CPU memory.
+ *
+ * Memory is owned by the caller. The caller must make sure the buffer is large enough and at least as large as SizeInBytes.
+ */
+struct NNE_API FTensorBindingCPU
+{
+	void*	Data;
+	uint64	SizeInBytes;
+};
+
+/**
+ * The interface of a model instance that can run synchronously from CPU memory.
+ *
+ * Use UE::NNE::IModelXXX::CreateModelInstance() to get a model instance.
+ * Use UE::NNE::GetRuntime<INNERuntimeXXX>(RuntimeName) to get a runtime capable of creating models.
+ */
+class NNE_API IModelInstanceRunSync
+{
+public:
+
+	using ESetInputTensorShapesStatus = EResultStatus;
+	using ERunSyncStatus = EResultStatus;
+
+	virtual ~IModelInstanceRunSync() = default;
+
+	/**
+	 * Get the input tensor descriptions as defined by the model, potentially with variable dimensions.
+	 *
+	 * @return An array containing a tensor descriptor for each input tensor of the model.
+	 */
+	virtual TConstArrayView<FTensorDesc> GetInputTensorDescs() const = 0;
+	
+	/**
+	 * Get the output tensor descriptions as defined by the model, potentially with variable dimensions.
+	 *
+	 * @return An array containing a tensor descriptor for each output tensor of the model.
+	 */
+	virtual TConstArrayView<FTensorDesc> GetOutputTensorDescs() const = 0;
+
+	/**
+	 * Get the input shapes.
+	 *
+	 * SetInputTensorShapes must be called prior of running a model.
+	 *
+	 * @return An array of input shapes or an empty array if SetInputTensorShapes has not been called.
+	 */
+	virtual TConstArrayView<FTensorShape> GetInputTensorShapes() const = 0;
+
+	/**
+	 * Getters for outputs shapes if they were already resolved.
+	 *
+	 * Output shapes might be resolved after a call to SetInputTensorShapes if the model and runtime supports it.
+	 * Otherwise they will be resolved while running the model
+	 *
+	 * @return An array of output shapes or an empty array if not resolved yet.
+	 */
+	virtual TConstArrayView<FTensorShape> GetOutputTensorShapes() const = 0;
+	
+	/**
+	 * Prepare the model to be run with the given input shape.
+	 *
+	 * The call is mandatory before a model can be run.
+	 * The function will run shape inference and resolve, if possible, the output shapes which can then be accessed by calling GetOutputTensorShapes().
+	 * This is a potentially expensive call and should be called lazily if possible.
+	 *
+	 * @param InInputShapes The input shapes to prepare the model with.
+	 * @return Status indicating success or failure.
+	 */
+	virtual ESetInputTensorShapesStatus SetInputTensorShapes(TConstArrayView<FTensorShape> InInputShapes) = 0;
+
+	/**
+	 * Evaluate the model synchronously.
+	 *
+	 * SetInputTensorShapes must be called prior to this call.
+	 * This function will block the calling thread until the inference is complete.
+	 * The caller owns the memory inside the bindings and must make sure that they are big enough.
+	 * Clients can call this function from an async task but must make sure the memory remains valid throughout the evaluation.
+	 *
+	 * @param InInputTensors An array containing tensor bindings for each input tensor with caller owned memory containing the input data.
+	 * @param InOutputTensors An array containing tensor bindings for each output tensor with caller owned memory big enough to contain the results on success.
+	 * @return Status indicating success or failure.
+	 */
+	virtual ERunSyncStatus RunSync(TConstArrayView<FTensorBindingCPU> InInputTensors, TConstArrayView<FTensorBindingCPU> InOutputTensors) = 0;
+};
+
+} // UE::NNE
\ No newline at end of file
