[NNE]
* Moved dll storage to methods/variables instead of on types. This is needed to get merged modules working

diff --git a/Engine/Source/Runtime/NNE/Public/NNEModelData.h b/Engine/Source/Runtime/NNE/Public/NNEModelData.h
--- a/Engine/Source/Runtime/NNE/Public/NNEModelData.h
+++ b/Engine/Source/Runtime/NNE/Public/NNEModelData.h
@@ -17,7 +17,7 @@ namespace UE::NNE
 	 * 
 	 * It allows runtimes to reference results of GetModelData() even if they outlive UNNEModelData.
 	 */
-	class NNE_API FSharedModelData
+	class FSharedModelData
 	{
 
 	public:
@@ -27,26 +27,26 @@ namespace UE::NNE
 		 * @param InData The shared buffer containing the model data. InData must be owned and the memory aligned with InMemoryAlignment.
 		 * @param InMemoryAlignment The memory alignment with which InData has been aligned. A value <= 1 indicates arbitrary memory alignment.
 		 */
-		FSharedModelData(FSharedBuffer InData, uint32 InMemoryAlignment);
+		NNE_API FSharedModelData(FSharedBuffer InData, uint32 InMemoryAlignment);
 
 		/**
 		 * Constructor to create empty data.
 		 */
-		FSharedModelData();
+		NNE_API FSharedModelData();
 
 		/**
 		 * Get a const array view on the shared data which is guaranteed to remain valid as long as this objects exists.
 		 *
 		 * @return A const array view of the shared data.
 		 */
-		TConstArrayView64<uint8> GetView() const;
+		NNE_API TConstArrayView64<uint8> GetView() const;
 
 		/**
 		 * Get the memory alignment with which the data has been aligned.
 		 *
 		 * @return Memory alignment with which the data has been aligned. A value <= 1 indicates arbitrary memory alignment.
 		 */
-		uint32 GetMemoryAlignment() const;
+		NNE_API uint32 GetMemoryAlignment() const;
 
 	private:
 
@@ -69,16 +69,16 @@ namespace UE::NNE
  * UNNEModelData assets store such model data as imported e.g. by the UNNEModelDataFactory class.
  * An INNERuntime object retrieved by UE::NNE::GetRuntime<T>(const FString& Name) can be used to create an inferable neural network model.
  */
-UCLASS(BlueprintType, Category = "NNE")
-class NNE_API UNNEModelData : public UObject
+UCLASS(BlueprintType, Category = "NNE", MinimalAPI)
+class UNNEModelData : public UObject
 {
 	GENERATED_BODY()
 
 public:
 
 	// UObject interface
-	virtual void GetAssetRegistryTags(FAssetRegistryTagsContext Context) const override;
-	virtual void Serialize(FArchive& Ar) override;
+	NNE_API virtual void GetAssetRegistryTags(FAssetRegistryTagsContext Context) const override;
+	NNE_API virtual void Serialize(FArchive& Ar) override;
 	// End of UObject interface
 
 	/**
@@ -90,7 +90,7 @@ public:
 	 * @param Buffer The raw binary file data of the imported model to be copied into this asset.
 	 * @param AdditionalBuffers Additional raw binary data of the model to be copied into this asset.
 	 */
-	void Init(const FString& Type, TConstArrayView64<uint8> Buffer, const TMap<FString, TConstArrayView64<uint8>>& AdditionalBuffers = TMap<FString, TConstArrayView64<uint8>>());
+	NNE_API void Init(const FString& Type, TConstArrayView64<uint8> Buffer, const TMap<FString, TConstArrayView64<uint8>>& AdditionalBuffers = TMap<FString, TConstArrayView64<uint8>>());
 
 	/**
 	 * In editor: Get the target runtimes this model data will be cooked for. An empty list means all runtimes.
@@ -98,14 +98,14 @@ public:
 	 *
 	 * @return The target runtimes names.
 	 */
-	TArrayView<const FString> GetTargetRuntimes() const;
+	NNE_API TArrayView<const FString> GetTargetRuntimes() const;
 
 	/**
 	 * Set the target runtimes this model data will be cooked for. An empty list means all runtimes.
 	 *
 	 * @param RuntimeNames The target runtimes names.
 	 */
-	void SetTargetRuntimes(TArrayView<const FString> RuntimeNames);
+	NNE_API void SetTargetRuntimes(TArrayView<const FString> RuntimeNames);
 
 	/**
 	 * Get the type of data inside FileData.
@@ -115,7 +115,7 @@ public:
 	 *
 	 * @return The FileType.
 	 */
-	FString GetFileType() const;
+	NNE_API FString GetFileType() const;
 
 	/**
 	 * Get read only access to FileData.
@@ -125,7 +125,7 @@ public:
 	 *
 	 * @return The FileData.
 	 */
-	TConstArrayView64<uint8> GetFileData() const;
+	NNE_API TConstArrayView64<uint8> GetFileData() const;
 
 	/**
 	 * Get read only access to AdditionalFileData.
@@ -135,14 +135,14 @@ public:
 	 *
 	 * @return The AdditionalFileData with a given Key if it exists and an empty view in standalone or when the key does not exist.
 	 */
-	TConstArrayView64<uint8> GetAdditionalFileData(const FString& Key) const;
+	NNE_API TConstArrayView64<uint8> GetAdditionalFileData(const FString& Key) const;
 
 	/**
 	 * Clears the FileData and the FileType.
 	 *
 	 * Caution, if the FileData is cleared, no more models can be created on runtimes that do not already have ModelData inside this asset.
 	 */
-	void ClearFileDataAndFileType();
+	NNE_API void ClearFileDataAndFileType();
 
 	/**
 	 * Get the FGuid identifying the FileData.
@@ -152,7 +152,7 @@ public:
 	 *
 	 * @return The FileId.
 	 */
-	FGuid GetFileId() const;
+	NNE_API FGuid GetFileId() const;
 
 	/**
 	 * Get the cached (editor) or cooked (game) optimized model data for a given runtime.
@@ -162,14 +162,14 @@ public:
 	 * @param RuntimeName The name of the runtime for which the data should be returned.
 	 * @return The optimized and runtime specific model data or an invalid TSharedPtr in case of failure.
 	 */
-	TSharedPtr<UE::NNE::FSharedModelData> GetModelData(const FString& RuntimeName);
+	NNE_API TSharedPtr<UE::NNE::FSharedModelData> GetModelData(const FString& RuntimeName);
 
 	/**
 	 * Clears the ModelData.
 	 *
 	 * Caution, if the ModelData is cleared, only runtimes that support cooking on the current platform can create new models from this asset.
 	 */
-	void ClearModelData();
+	NNE_API void ClearModelData();
 
 private:
 	/**
diff --git a/Engine/Source/Runtime/NNE/Public/NNERuntime.h b/Engine/Source/Runtime/NNE/Public/NNERuntime.h
--- a/Engine/Source/Runtime/NNE/Public/NNERuntime.h
+++ b/Engine/Source/Runtime/NNE/Public/NNERuntime.h
@@ -16,8 +16,8 @@ namespace UE::NNE
 	class FSharedModelData;
 }
 
-UINTERFACE()
-class NNE_API UNNERuntime : public UInterface
+UINTERFACE(MinimalAPI)
+class UNNERuntime : public UInterface
 {
 	GENERATED_BODY()
 };
@@ -29,7 +29,7 @@ class NNE_API UNNERuntime : public UInterface
  * The model data is then stored inside UNNEModelData which then can be used in game to create an inferable model.
  * See INNERuntimeCPU, INNERuntimeGPU and INNERuntimeRDG on how to create a model given a UNNEModelData asset.
  */
-class NNE_API INNERuntime
+class INNERuntime
 {
 	GENERATED_BODY()
 
diff --git a/Engine/Source/Runtime/NNE/Public/NNERuntimeCPU.h b/Engine/Source/Runtime/NNE/Public/NNERuntimeCPU.h
--- a/Engine/Source/Runtime/NNE/Public/NNERuntimeCPU.h
+++ b/Engine/Source/Runtime/NNE/Public/NNERuntimeCPU.h
@@ -19,7 +19,7 @@ namespace UE::NNE
  * Use UE::NNE::IModelCPU::CreateModelInstance() to get a model instance.
  * Use UE::NNE::GetRuntime<INNERuntimeCPU>(RuntimeName) to get a runtime capable of creating CPU models.
  */
-class NNE_API IModelInstanceCPU : public IModelInstanceRunSync
+class IModelInstanceCPU : public IModelInstanceRunSync
 {
 };
 
@@ -28,7 +28,7 @@ class NNE_API IModelInstanceCPU : public IModelInstanceRunSync
  *
  * Use UE::NNE::GetRuntime<INNERuntimeCPU>(RuntimeName) to get a runtime capable of creating CPU models.
  */
-class NNE_API IModelCPU
+class IModelCPU
 {
 public:
 
@@ -47,8 +47,8 @@ public:
 
 } // UE::NNE
 
-UINTERFACE()
-class NNE_API UNNERuntimeCPU : public UInterface
+UINTERFACE(MinimalAPI)
+class UNNERuntimeCPU : public UInterface
 {
 	GENERATED_BODY()
 };
@@ -58,7 +58,7 @@ class NNE_API UNNERuntimeCPU : public UInterface
  *
  * Call UE::NNE::GetRuntime<INNERuntimeCPU>(RuntimeName) to get a runtime implementing this interface.
  */
-class NNE_API INNERuntimeCPU
+class INNERuntimeCPU
 {
 	GENERATED_BODY()
 	
diff --git a/Engine/Source/Runtime/NNE/Public/NNERuntimeGPU.h b/Engine/Source/Runtime/NNE/Public/NNERuntimeGPU.h
--- a/Engine/Source/Runtime/NNE/Public/NNERuntimeGPU.h
+++ b/Engine/Source/Runtime/NNE/Public/NNERuntimeGPU.h
@@ -21,7 +21,7 @@ typedef FTensorBindingCPU FTensorBindingGPU;
  * Use UE::NNE::IModelGPU::CreateModelInstance() to get a model instance.
  * Use UE::NNE::GetRuntime<INNERuntimeGPU>(RuntimeName) to get a runtime capable of creating GPU models.
  */
-class NNE_API IModelInstanceGPU : public IModelInstanceRunSync
+class IModelInstanceGPU : public IModelInstanceRunSync
 {
 };
 
@@ -30,7 +30,7 @@ class NNE_API IModelInstanceGPU : public IModelInstanceRunSync
  *
  * Use UE::NNE::GetRuntime<INNERuntimeGPU>(RuntimeName) to get a runtime capable of creating GPU models.
  */
-class NNE_API IModelGPU
+class IModelGPU
 {
 public:
 
@@ -49,8 +49,8 @@ public:
 
 } // UE::NNE
 
-UINTERFACE()
-class NNE_API UNNERuntimeGPU : public UInterface
+UINTERFACE(MinimalAPI)
+class UNNERuntimeGPU : public UInterface
 {
 	GENERATED_BODY()
 };
@@ -60,7 +60,7 @@ class NNE_API UNNERuntimeGPU : public UInterface
  *
  * Call UE::NNE::GetRuntime<INNERuntimeGPU>(RuntimeName) to get a runtime implementing this interface.
  */
-class NNE_API INNERuntimeGPU
+class INNERuntimeGPU
 {
 	GENERATED_BODY()
 	
diff --git a/Engine/Source/Runtime/NNE/Public/NNERuntimeNPU.h b/Engine/Source/Runtime/NNE/Public/NNERuntimeNPU.h
--- a/Engine/Source/Runtime/NNE/Public/NNERuntimeNPU.h
+++ b/Engine/Source/Runtime/NNE/Public/NNERuntimeNPU.h
@@ -19,7 +19,7 @@ namespace UE::NNE
  * Use UE::NNE::IModelNPU::CreateModelInstance() to get a model instance.
  * Use UE::NNE::GetRuntime<INNERuntimeNPU>(RuntimeName) to get a runtime capable of creating NPU models.
  */
-class NNE_API IModelInstanceNPU : public IModelInstanceRunSync
+class IModelInstanceNPU : public IModelInstanceRunSync
 {
 };
 
@@ -28,7 +28,7 @@ class NNE_API IModelInstanceNPU : public IModelInstanceRunSync
  *
  * Use UE::NNE::GetRuntime<INNERuntimeNPU>(RuntimeName) to get a runtime capable of creating NPU models.
  */
-class NNE_API IModelNPU
+class IModelNPU
 {
 public:
 
@@ -47,8 +47,8 @@ public:
 
 } // UE::NNE
 
-UINTERFACE()
-class NNE_API UNNERuntimeNPU : public UInterface
+UINTERFACE(MinimalAPI)
+class UNNERuntimeNPU : public UInterface
 {
 	GENERATED_BODY()
 };
@@ -58,7 +58,7 @@ class NNE_API UNNERuntimeNPU : public UInterface
  *
  * Call UE::NNE::GetRuntime<INNERuntimeNPU>(RuntimeName) to get a runtime implementing this interface.
  */
-class NNE_API INNERuntimeNPU
+class INNERuntimeNPU
 {
 	GENERATED_BODY()
 	
diff --git a/Engine/Source/Runtime/NNE/Public/NNERuntimeRDG.h b/Engine/Source/Runtime/NNE/Public/NNERuntimeRDG.h
--- a/Engine/Source/Runtime/NNE/Public/NNERuntimeRDG.h
+++ b/Engine/Source/Runtime/NNE/Public/NNERuntimeRDG.h
@@ -24,7 +24,7 @@ namespace UE::NNE
  *
  * Memory is owned by the caller. The caller must make sure the buffer is large enough.
  */
-struct NNE_API FTensorBindingRDG
+struct FTensorBindingRDG
 {
 	FRDGBufferRef Buffer;
 };
@@ -35,7 +35,7 @@ struct NNE_API FTensorBindingRDG
  * Use UE::NNE::IModelRDG::CreateModelInstance() to get a model instance.
  * Use UE::NNE::GetRuntime<INNERuntimeRDG>(RuntimeName) to get a runtime capable of creating RDG models.
  */
-class NNE_API IModelInstanceRDG
+class IModelInstanceRDG
 {
 public:
 
@@ -110,7 +110,7 @@ public:
  *
  * Use UE::NNE::GetRuntime<INNERuntimeRDG>(RuntimeName) to get a runtime capable of creating RDG models.
  */
-class NNE_API IModelRDG
+class IModelRDG
 {
 public:
 
@@ -129,8 +129,8 @@ public:
 
 } // UE::NNE
 
-UINTERFACE()
-class NNE_API UNNERuntimeRDG : public UInterface
+UINTERFACE(MinimalAPI)
+class UNNERuntimeRDG : public UInterface
 {
 	GENERATED_BODY()
 };
@@ -140,7 +140,7 @@ class NNE_API UNNERuntimeRDG : public UInterface
  *
  * Call UE::NNE::GetRuntime<INNERuntimeRDG>(RuntimeName) to get a runtime implementing this interface.
  */
-class NNE_API INNERuntimeRDG
+class INNERuntimeRDG
 {
 	GENERATED_BODY()
 
diff --git a/Engine/Source/Runtime/NNE/Public/NNERuntimeRunSync.h b/Engine/Source/Runtime/NNE/Public/NNERuntimeRunSync.h
--- a/Engine/Source/Runtime/NNE/Public/NNERuntimeRunSync.h
+++ b/Engine/Source/Runtime/NNE/Public/NNERuntimeRunSync.h
@@ -15,7 +15,7 @@ namespace UE::NNE
  *
  * Memory is owned by the caller. The caller must make sure the buffer is large enough and at least as large as SizeInBytes.
  */
-struct NNE_API FTensorBindingCPU
+struct FTensorBindingCPU
 {
 	void*	Data;
 	uint64	SizeInBytes;
@@ -27,7 +27,7 @@ struct NNE_API FTensorBindingCPU
  * Use UE::NNE::IModelXXX::CreateModelInstance() to get a model instance.
  * Use UE::NNE::GetRuntime<INNERuntimeXXX>(RuntimeName) to get a runtime capable of creating models.
  */
-class NNE_API IModelInstanceRunSync
+class IModelInstanceRunSync
 {
 public:
 
diff --git a/Engine/Source/Runtime/NNE/Public/NNETypes.h b/Engine/Source/Runtime/NNE/Public/NNETypes.h
--- a/Engine/Source/Runtime/NNE/Public/NNETypes.h
+++ b/Engine/Source/Runtime/NNE/Public/NNETypes.h
@@ -48,7 +48,7 @@ namespace UE::NNE
 	 * 
 	 * The variable dimensions are represented by -1 values. 
 	 */
-	class NNE_API FSymbolicTensorShape
+	class FSymbolicTensorShape
 	{
 	public:
 		/**
@@ -78,7 +78,7 @@ namespace UE::NNE
 		 * @param Data an array of dimensions.
 		 * @return a symbolic tensor shape with the given dimensions.
 		 */
-		static FSymbolicTensorShape Make(TConstArrayView<int32> Data);
+		static NNE_API FSymbolicTensorShape Make(TConstArrayView<int32> Data);
 
 		/**
 		 * Get the dimensions of the symbolic tensor shape.
@@ -99,7 +99,7 @@ namespace UE::NNE
 		 *
 		 * @return true if the symbolic tensor shape is concrete, false otherwise.
 		 */
-		bool IsConcrete() const;
+		NNE_API bool IsConcrete() const;
 
 		/**
 		 * Check if the symbolic tensor shape is equal to another one.
@@ -107,7 +107,7 @@ namespace UE::NNE
 		 * @param OtherShape the other symbolic tensor shape.
 		 * @return true if the symbolic tensor shape is equal to the other one, false otherwise.
 		 */
-		bool operator==(const FSymbolicTensorShape& OtherShape) const;
+		NNE_API bool operator==(const FSymbolicTensorShape& OtherShape) const;
 		
 		/**
 		 * Check if the symbolic tensor shape is different from another one.
@@ -115,14 +115,14 @@ namespace UE::NNE
 		 * @param OtherShape the other symbolic tensor shape.
 		 * @return true if the symbolic tensor shape is different from the other one, false otherwise.
 		 */
-		bool operator!=(const FSymbolicTensorShape& OtherShape) const;
+		NNE_API bool operator!=(const FSymbolicTensorShape& OtherShape) const;
 
 		/**
 		 * Assign this symbolic tensor shape from another one.
 		 * 
 		 * @param OtherShape the other symbolic tensor shape to copy from.
 		 */
-		void operator=(const FSymbolicTensorShape& OtherShape);
+		NNE_API void operator=(const FSymbolicTensorShape& OtherShape);
 	};
 
 	/**
@@ -130,7 +130,7 @@ namespace UE::NNE
 	 *
 	 * Concrete tensor shapes are well defined through strictly positive values and thus have a defined volume.
 	 */
-	class NNE_API FTensorShape
+	class FTensorShape
 	{
 	public:
 		/**
@@ -159,7 +159,7 @@ namespace UE::NNE
 		 * @param Data an array of dimensions.
 		 * @return a tensor shape with the given dimensions.
 		 */
-		static FTensorShape Make(TConstArrayView<uint32> Data);
+		static NNE_API FTensorShape Make(TConstArrayView<uint32> Data);
 		
 		/**
 		 * Construct a tensor shape from a symbolic tensor shape.
@@ -169,7 +169,7 @@ namespace UE::NNE
 		 * @param SymbolicShape a symbolic tensor shape.
 		 * @return a tensor shape with the given dimensions.
 		 */
-		static FTensorShape MakeFromSymbolic(const FSymbolicTensorShape& SymbolicShape);
+		static NNE_API FTensorShape MakeFromSymbolic(const FSymbolicTensorShape& SymbolicShape);
 
 		/**
 		 * Get the dimensions of the tensor shape.
@@ -190,7 +190,7 @@ namespace UE::NNE
 		 *
 		 * @return the number of element in this tensor shape.
 		 */
-		uint64 Volume() const;
+		NNE_API uint64 Volume() const;
 		
 		/**
 		 * Check if this tensor shape is part of the ensemble of shapes defined by the symbolic shape.
@@ -199,7 +199,7 @@ namespace UE::NNE
 		 *
 		 * @return true if this tensor shape is compatible with the symbolic shape, false otherwise.
 		 */
-		bool IsCompatibleWith(const FSymbolicTensorShape& SymbolicShape) const;
+		NNE_API bool IsCompatibleWith(const FSymbolicTensorShape& SymbolicShape) const;
 
 		/**
 		 * Check if the tensor shape is equal to another one.
@@ -207,7 +207,7 @@ namespace UE::NNE
 		 * @param OtherShape the other tensor shape.
 		 * @return true if the tensor shape is equal to the other one, false otherwise.
 		 */
-		bool operator==(const FTensorShape& OtherShape) const;
+		NNE_API bool operator==(const FTensorShape& OtherShape) const;
 		
 		/**
 		 * Check if the tensor shape is different from another one.
@@ -215,13 +215,13 @@ namespace UE::NNE
 		 * @param OtherShape the other tensor shape.
 		 * @return true if the tensor shape is different from the other one, false otherwise.
 		 */
-		bool operator!=(const FTensorShape& OtherShape) const;
+		NNE_API bool operator!=(const FTensorShape& OtherShape) const;
 
 		/**
 		 * Assign this tensor shape from another one.
 		 * @param OtherShape the other tensor shape to copy from.
 		 */
-		void operator=(const FTensorShape& OtherShape);
+		NNE_API void operator=(const FTensorShape& OtherShape);
 	};
 
 	/**
@@ -230,7 +230,7 @@ namespace UE::NNE
 	 * A tensor is described by its name, the type of data it contains and it's shape.
 	 * Since input and output tensors of a neural network can have dynamic shapes, Shape is symbolic.
 	 */
-	class NNE_API FTensorDesc
+	class FTensorDesc
 	{
 	private:
 		FString					Name;
@@ -248,7 +248,7 @@ namespace UE::NNE
 		 * @param DataType the type of data the tensor contains.
 		 * @return a tensor shape with the given dimensions.
 		 */
-		static FTensorDesc Make(const FString& Name, const FSymbolicTensorShape& Shape, ENNETensorDataType DataType);
+		static NNE_API FTensorDesc Make(const FString& Name, const FSymbolicTensorShape& Shape, ENNETensorDataType DataType);
 
 		/**
 		 * Get the name of the tensor
